Copying //UE4/Dev-Physics to Dev-Main (//UE4/Dev-Main) [at] 7703071

diff --git a/Engine/Source/Runtime/PhysicsCore/Public/PhysXInterfaceWrapperCore.h b/Engine/Source/Runtime/PhysicsCore/Public/PhysXInterfaceWrapperCore.h
--- a/Engine/Source/Runtime/PhysicsCore/Public/PhysXInterfaceWrapperCore.h
+++ b/Engine/Source/Runtime/PhysicsCore/Public/PhysXInterfaceWrapperCore.h
@@ -1,7 +1,7 @@
 // Copyright 1998-2019 Epic Games, Inc. All Rights Reserved.
 #pragma once
 
-#if WITH_PHYSX && PHYSICS_INTERFACE_PHYSX
+#if WITH_PHYSX
 #include "PhysXPublicCore.h"
 #include "PhysicsInterfaceUtilsCore.h"
 #include "PhysicsInterfaceWrapperShared.h"
@@ -49,58 +49,6 @@
 	return (PxQueryFlags)Result;
 }
 
-FORCEINLINE bool HadInitialOverlap(const PxLocationHit& Hit)
-{
-	return Hit.hadInitialOverlap();
-}
-
-FORCEINLINE PxShape* GetShape(const PxLocationHit& Hit)
-{
-	return Hit.shape;
-}
-
-FORCEINLINE PxShape* GetShape(const PxOverlapHit& Hit)
-{
-	return Hit.shape;
-}
-
-FORCEINLINE PxRigidActor* GetActor(const PxLocationHit& Hit)
-{
-	return Hit.actor;
-}
-
-FORCEINLINE PxRigidActor* GetActor(const PxOverlapHit& Hit)
-{
-	return Hit.actor;
-}
-
-FORCEINLINE float GetDistance(const PxLocationHit& Hit)
-{
-	return Hit.distance;
-}
-
-template <typename HitType>
-HitType* GetBlock(PxHitCallback<HitType>& Callback)
-{
-	return &Callback.block;
-}
-
-template <typename HitType>
-bool GetHasBlock(const PxHitCallback<HitType>& Callback)
-{
-	return Callback.hasBlock;
-}
-
-FORCEINLINE FVector GetPosition(const PxLocationHit& Hit)
-{
-	return P2UVector(Hit.position);
-}
-
-FORCEINLINE FVector GetNormal(const PxLocationHit& Hit)
-{
-	return P2UVector(Hit.normal);
-}
-
 FORCEINLINE PxHitFlags U2PHitFlags(const FHitFlags& Flags)
 {
 	uint32 Result = 0;
@@ -173,6 +121,61 @@
 	return Result.HitFlags;
 }
 
+namespace PhysXInterface
+{
+
+FORCEINLINE bool HadInitialOverlap(const PxLocationHit& Hit)
+{
+	return Hit.hadInitialOverlap();
+}
+
+FORCEINLINE PxShape* GetShape(const PxLocationHit& Hit)
+{
+	return Hit.shape;
+}
+
+FORCEINLINE PxShape* GetShape(const PxOverlapHit& Hit)
+{
+	return Hit.shape;
+}
+
+FORCEINLINE PxRigidActor* GetActor(const PxLocationHit& Hit)
+{
+	return Hit.actor;
+}
+
+FORCEINLINE PxRigidActor* GetActor(const PxOverlapHit& Hit)
+{
+	return Hit.actor;
+}
+
+FORCEINLINE float GetDistance(const PxLocationHit& Hit)
+{
+	return Hit.distance;
+}
+
+template <typename HitType>
+HitType* GetBlock(PxHitCallback<HitType>& Callback)
+{
+	return &Callback.block;
+}
+
+template <typename HitType>
+bool GetHasBlock(const PxHitCallback<HitType>& Callback)
+{
+	return Callback.hasBlock;
+}
+
+FORCEINLINE FVector GetPosition(const PxLocationHit& Hit)
+{
+	return P2UVector(Hit.position);
+}
+
+FORCEINLINE FVector GetNormal(const PxLocationHit& Hit)
+{
+	return P2UVector(Hit.normal);
+}
+
 FORCEINLINE FHitFlags GetFlags(const PxLocationHit& Hit)
 {
 	return P2UHitFlags(Hit.flags);
@@ -231,56 +234,4 @@
 	return (faceIndex == 0xFFFFffff);
 }
 
-FORCEINLINE uint32 GetTriangleMeshExternalFaceIndex(const PxShape& Shape, uint32 InternalFaceIndex)
-{
-	PxTriangleMeshGeometry TriangleMeshGeometry;
-	if (Shape.getTriangleMeshGeometry(TriangleMeshGeometry))
-	{
-		if (TriangleMeshGeometry.triangleMesh && InternalFaceIndex < TriangleMeshGeometry.triangleMesh->getNbTriangles())
-		{
-			if (const PxU32* TriangleRemap = TriangleMeshGeometry.triangleMesh->getTrianglesRemap())
-			{
-				return TriangleRemap[InternalFaceIndex];
-			}
-		}
-	}
-
-	return GetInvalidPhysicsFaceIndex();
-}
-
-FORCEINLINE float GetRadius(const PxCapsuleGeometry& Capsule)
-{
-	return Capsule.radius;
-}
-
-FORCEINLINE float GetHalfHeight(const PxCapsuleGeometry& Capsule)
-{
-	return Capsule.halfHeight;
-}
-
-FORCEINLINE PxTransform GetGlobalPose(const PxRigidActor& RigidActor)
-{
-	return RigidActor.getGlobalPose();
-}
-
-FORCEINLINE uint32 GetNumShapes(const PxRigidActor& RigidActor)
-{
-	return RigidActor.getNbShapes();
-}
-
-FORCEINLINE void GetShapes(const PxRigidActor& RigidActor, PxShape** ShapesBuffer, uint32 NumShapes)
-{
-	RigidActor.getShapes(ShapesBuffer, sizeof(ShapesBuffer[0]) * NumShapes);
-}
-
-FORCEINLINE void SetActor(PxActorShape& Hit, PxRigidActor* Actor)
-{
-	Hit.actor = Actor;
-}
-
-FORCEINLINE void SetShape(PxActorShape& Hit, PxShape* Shape)
-{
-	Hit.shape = Shape;
-}
-
 template <typename T>
@@ -353,34 +304,6 @@
 	}
 };
 
-template <typename HitType>
-FORCEINLINE_DEBUGGABLE bool Insert(PxHitCallback<HitType>& Callback, const HitType& Hit, bool bBlocking)
-{
-	if (!Callback.hasBlock || GetDistance(Hit) < GetDistance(Callback.block))
-	{
-		if (bBlocking)
-		{
-			Callback.block = Hit;
-			Callback.hasBlock = true;
-		}
-		else
-		{
-			if (Callback.maxNbTouches > 0)
-			{
-				Callback.processTouches(&Hit, 1);
 
-			}
-		}
-	}
-
-	return true;
-}
-
-template <typename HitType>
-FORCEINLINE bool InsertOverlap(PxHitCallback<HitType>& Callback, const HitType& Hit)
-{
-	return Callback.processTouches(&Hit, 1);
-}
-
 /** We use this struct so that if no conversion is needed in another API, we can avoid the copy (if we think that's critical) */
 struct FPhysicsRaycastInputAdapater
@@ -423,33 +346,62 @@
 	PxTransform GeomPose;
 };
 
-template <typename HitType>
-FORCEINLINE void SetBlock(PxHitCallback<HitType>& Callback, const HitType& Hit)
+FORCEINLINE uint32 GetTriangleMeshExternalFaceIndex(const PxShape& Shape, uint32 InternalFaceIndex)
 {
-	Callback.block = Hit;
-}
-
-template <typename HitType>
-FORCEINLINE void SetHasBlock(PxHitCallback<HitType>& Callback, bool bHasBlock)
-{
-	Callback.hasBlock = bHasBlock;
-}
-
-template <typename HitType>
-FORCEINLINE void ProcessTouches(PxHitCallback<HitType>& Callback, const TArray<HitType>& TouchingHits)
-{
-	Callback.processTouches(TouchingHits.GetData(), TouchingHits.Num());
-}
-
-template <typename HitType>
-FORCEINLINE void FinalizeQuery(PxHitCallback<HitType>& Callback)
-{
-	if (Callback.nbTouches)
+	PxTriangleMeshGeometry TriangleMeshGeometry;
+	if (Shape.getTriangleMeshGeometry(TriangleMeshGeometry))
 	{
-		Callback.processTouches(Callback.touches, Callback.nbTouches);
+		if (TriangleMeshGeometry.triangleMesh && InternalFaceIndex < TriangleMeshGeometry.triangleMesh->getNbTriangles())
+		{
+			if (const PxU32* TriangleRemap = TriangleMeshGeometry.triangleMesh->getTrianglesRemap())
+			{
+				return TriangleRemap[InternalFaceIndex];
+			}
+		}
 	}
 
-	Callback.finalizeQuery();
+	return PhysXInterface::GetInvalidPhysicsFaceIndex();
 }
 
+FORCEINLINE float GetRadius(const PxCapsuleGeometry& Capsule)
+{
+	return Capsule.radius;
+}
+
+FORCEINLINE float GetHalfHeight(const PxCapsuleGeometry& Capsule)
+{
+	return Capsule.halfHeight;
+}
+
+FORCEINLINE PxTransform GetGlobalPose(const PxRigidActor& RigidActor)
+{
+	return RigidActor.getGlobalPose();
+}
+
+FORCEINLINE uint32 GetNumShapes(const PxRigidActor& RigidActor)
+{
+	return RigidActor.getNbShapes();
+}
+
+FORCEINLINE void GetShapes(const PxRigidActor& RigidActor, PxShape** ShapesBuffer, uint32 NumShapes)
+{
+	RigidActor.getShapes(ShapesBuffer, sizeof(ShapesBuffer[0]) * NumShapes);
+}
+
+FORCEINLINE void SetActor(PxActorShape& Hit, PxRigidActor* Actor)
+{
+	Hit.actor = Actor;
+}
+
+FORCEINLINE void SetShape(PxActorShape& Hit, PxShape* Shape)
+{
+	Hit.shape = Shape;
+}
+
+}
+
+#if PHYSICS_INTERFACE_PHYSX
+using namespace PhysXInterface;
 #endif
+
+#endif
diff --git a/Engine/Source/Runtime/PhysicsCore/Public/PhysicsCoreTypes.h b/Engine/Source/Runtime/PhysicsCore/Public/PhysicsCoreTypes.h
--- a/Engine/Source/Runtime/PhysicsCore/Public/PhysicsCoreTypes.h
+++ b/Engine/Source/Runtime/PhysicsCore/Public/PhysicsCoreTypes.h
@@ -3,7 +3,10 @@
 #pragma once
 
 #include "CoreMinimal.h"
+
+#if INCLUDE_CHAOS
 #include "Chaos/PBDRigidParticles.h"
+#endif
 
 #include "PhysicsCoreTypes.generated.h"
 
@@ -35,5 +38,6 @@ enum class EChaosBufferMode : uint8
 	Invalid UMETA(Hidden)
 };
 
+#if INCLUDE_CHAOS
 typedef Chaos::TPBDRigidParticles<float, 3> FParticlesType;
-
+#endif
diff --git a/Engine/Source/Runtime/PhysicsCore/Public/PhysicsInterfaceDeclaresCore.h b/Engine/Source/Runtime/PhysicsCore/Public/PhysicsInterfaceDeclaresCore.h
--- a/Engine/Source/Runtime/PhysicsCore/Public/PhysicsInterfaceDeclaresCore.h
+++ b/Engine/Source/Runtime/PhysicsCore/Public/PhysicsInterfaceDeclaresCore.h
@@ -14,19 +14,8 @@
 #define WITH_IMMEDIATE_PHYSX 0
 #endif
 
-#ifndef PHYSICS_INTERFACE_LLIMMEDIATE
-#define PHYSICS_INTERFACE_LLIMMEDIATE 0
-#endif
-
-
 #if PHYSICS_INTERFACE_PHYSX
 
-#if WITH_IMMEDIATE_PHYSX
-
-
-
-#else
-
 namespace physx
 {
 	struct PxLocationHit;
@@ -46,8 +35,10 @@ namespace physx
 	struct PxHitBuffer;
 }
 
-template<typename T>
-class FSingleHitBuffer;
+namespace PhysXInterface
+{
+	struct FDummyPhysType {};
+}
 
 using FHitLocation = physx::PxLocationHit;
 using FHitSweep = physx::PxSweepHit;
@@ -56,6 +47,7 @@ using FHitOverlap = physx::PxOverlapHit;
 using FPhysicsQueryHit = physx::PxQueryHit;
 
 using FPhysicsTransform = physx::PxTransform;
+using FPhysTypeDummy = PhysXInterface::FDummyPhysType;
 
 using FPhysicsShape = physx::PxShape;
 using FPhysicsGeometry = physx::PxGeometry;
@@ -63,58 +55,10 @@ using FPhysicsCapsuleGeometry = physx::PxCapsuleGeometry;
 using FPhysicsMaterial = physx::PxMaterial;
 using FPhysicsActor = physx::PxRigidActor;
 
-using FPhysicsSweepBuffer = FSingleHitBuffer<physx::PxSweepHit>;
-using FPhysicsRaycastBuffer = FSingleHitBuffer<physx::PxRaycastHit>;
-
-#endif
-
-#elif PHYSICS_INTERFACE_LLIMMEDIATE
-
-namespace physx
-{
-	//struct PxLocationHit;
-	//struct PxSweepHit;
-	//struct PxRaycastHit;
-	//struct PxOverlapHit;
-	//
-	//class PxShape;
-	class PxGeometry;
-	class PxCapsuleGeometry;
-	//class PxMaterial;
-	//class PxRigidActor;
-	//
-	//template<typename T>
-	//struct PxHitBuffer;
-
-	//struct PxSweepBuffer;
-	//struct PxRaycastBuffer;
-}
-
-// Temporary dummy types until SQ implemented
-struct FPhysTypeDummy {};
-struct FPhysActorDummy {};
-template<typename T>
-struct FCallbackDummy;
-
-using FHitLocation = FPhysTypeDummy;
-using FHitSweep = FPhysTypeDummy;
-using FHitRaycast = FPhysTypeDummy;
-using FHitOverlap = FPhysTypeDummy;
-using FPhysicsQueryHit = FPhysTypeDummy;
-
-using FPhysicsTransform = FTransform;
-
-using FPhysicsShape = FPhysTypeDummy;
-using FPhysicsGeometry = physx::PxGeometry;
-using FPhysicsCapsuleGeometry = physx::PxCapsuleGeometry;
-using FPhysicsMaterial = FPhysTypeDummy;
-using FPhysicsActor = FPhysActorDummy;
-
-using FPhysicsSweepBuffer = FCallbackDummy<FHitSweep>;
-using FPhysicsRaycastBuffer = FCallbackDummy<FHitRaycast>;
-
 #elif WITH_CHAOS
 
+#include "ChaosSQTypes.h"
+
 namespace Chaos
 {
 	template<class T, int d>
@@ -122,34 +66,51 @@ namespace Chaos
 
 	template<class T>
 	class TCapsule;
+
+	template <typename T, int d>
+	class TGeometryParticle;
+
+	template <typename T, int d>
+	class TPerShapeData;
 }
 
 // Temporary dummy types until SQ implemented
-struct FPhysTypeDummy {};
-struct FPhysSweepDummy : public FPhysTypeDummy {};
-struct FPhysRaycastDummy : public FPhysTypeDummy {};
-struct FPhysActorDummy {};
-
+namespace ChaosInterface
+{
+	struct FDummyPhysType;
+	struct FDummyPhysActor;
+	template<typename T> struct FDummyCallback;
+}
+using FPhysTypeDummy = ChaosInterface::FDummyPhysType;
+using FPhysActorDummy = ChaosInterface::FDummyPhysActor;
 
 template<typename T>
-struct FCallbackDummy;
+using FCallbackDummy = ChaosInterface::FDummyCallback<T>;
 
-using FHitLocation = FPhysTypeDummy;
-using FHitSweep = FPhysSweepDummy;
-using FHitRaycast = FPhysRaycastDummy;
-using FHitOverlap = FPhysTypeDummy;
-using FPhysicsQueryHit = FPhysTypeDummy;
+struct FTransform;
+
+using FHitLocation = ChaosInterface::FLocationHit;
+using FHitSweep = ChaosInterface::FSweepHit;
+using FHitRaycast = ChaosInterface::FRaycastHit;
+using FHitOverlap = ChaosInterface::FOverlapHit;
+using FPhysicsQueryHit = ChaosInterface::FQueryHit;
 
 using FPhysicsTransform = FTransform;
 
-using FPhysicsShape = Chaos::TImplicitObject<float, 3>;
+using FPhysicsShape = Chaos::TPerShapeData<float, 3>;
 using FPhysicsGeometry = Chaos::TImplicitObject<float, 3>;
 using FPhysicsCapsuleGeometry = Chaos::TCapsule<float>;
 using FPhysicsMaterial = FPhysTypeDummy;
-using FPhysicsActor = FPhysActorDummy;
+using FPhysicsActor = Chaos::TGeometryParticle<float,3>;
+
+template <typename T>
+using FPhysicsHitCallback = ChaosInterface::FSQHitBuffer<T>;
+
+template <typename T>
+using FSingleHitBuffer = ChaosInterface::FSQSingleHitBuffer<T>;
 
-using FPhysicsSweepBuffer = FCallbackDummy<FHitSweep>;
-using FPhysicsRaycastBuffer = FCallbackDummy<FHitRaycast>;
+template <typename T>
+using FDynamicHitBuffer = ChaosInterface::FSQHitBuffer<T>;
 
 #else
 
@@ -170,6 +131,8 @@ struct FPhysicsCommand_ImmediatePhysX;
 struct FPhysicsShapeReference_ImmediatePhysX;
 struct FPhysicsMaterialReference_ImmediatePhysX;
 struct FPhysicsGeometryCollection_ImmediatePhysX;
+struct FPhysXShapeAdapter;
+struct FPhysxUserData;
 
 typedef FPhysicsActorReference_ImmediatePhysX		FPhysicsActorHandle;
 typedef FPhysicsConstraintReference_ImmediatePhysX	FPhysicsConstraintHandle;
@@ -178,8 +141,12 @@ typedef FPhysScene_ImmediatePhysX					FPhysScene;
 typedef FPhysicsAggregateReference_ImmediatePhysX	FPhysicsAggregateHandle;
 typedef FPhysicsCommand_ImmediatePhysX				FPhysicsCommand;
 typedef FPhysicsShapeReference_ImmediatePhysX		FPhysicsShapeHandle;
-typedef FPhysicsGeometryCollection_ImmediatePhysX   FPhysicsGeometryCollection;
-typedef FPhysicsMaterialReference_ImmediatePhysX    FPhysicsMaterialHandle;
+typedef FPhysicsGeometryCollection_ImmediatePhysX	FPhysicsGeometryCollection;
+typedef FPhysicsMaterialReference_ImmediatePhysX	FPhysicsMaterialHandle;
+typedef FPhysXShapeAdapter							FPhysicsShapeAdapter;
+typedef FPhysxUserData								FPhysicsUserData;
+
+inline FPhysicsActorHandle DefaultPhysicsActorHandle() { return FPhysicsActorHandle(); }
 
 #else
 
@@ -192,6 +159,8 @@ struct FPhysicsCommand_PhysX;
 struct FPhysicsShapeHandle_PhysX;
 struct FPhysicsGeometryCollection_PhysX;
 struct FPhysicsMaterialHandle_PhysX;
+struct FPhysXShapeAdapter;
+struct FPhysxUserData;
 
 typedef FPhysicsActorHandle_PhysX			FPhysicsActorHandle;
 typedef FPhysicsConstraintHandle_PhysX		FPhysicsConstraintHandle;
@@ -202,55 +171,39 @@ typedef FPhysicsCommand_PhysX				FPhysicsCommand;
 typedef FPhysicsShapeHandle_PhysX			FPhysicsShapeHandle;
 typedef FPhysicsGeometryCollection_PhysX	FPhysicsGeometryCollection;
 typedef FPhysicsMaterialHandle_PhysX		FPhysicsMaterialHandle;
+typedef FPhysXShapeAdapter					FPhysicsShapeAdapter;
+typedef FPhysxUserData						FPhysicsUserData;
 
-
+extern FPhysicsActorHandle DefaultPhysicsActorHandle();
 
 #endif
 
-#elif PHYSICS_INTERFACE_LLIMMEDIATE
-
-struct FPhysicsActorHandle_LLImmediate;
-struct FPhysicsAggregateHandle_LLImmediate;
-struct FPhysicsConstraintHandle_LLImmediate;
-struct FPhysicsShapeHandle_LLImmediate;
-struct FPhysicsCommand_LLImmediate;
-struct FPhysicsGeometryCollection_LLImmediate;
-struct FPhysicsMaterialHandle_LLImmediate;
-class FPhysInterface_LLImmediate;
-
-typedef FPhysicsActorHandle_LLImmediate           FPhysicsActorHandle;
-typedef FPhysicsConstraintHandle_LLImmediate      FPhysicsConstraintHandle;
-typedef FPhysInterface_LLImmediate                FPhysicsInterface;
-typedef FPhysInterface_LLImmediate                FPhysScene;
-typedef FPhysicsAggregateHandle_LLImmediate       FPhysicsAggregateHandle;
-typedef FPhysicsCommand_LLImmediate               FPhysicsCommand;
-typedef FPhysicsShapeHandle_LLImmediate           FPhysicsShapeHandle;
-typedef FPhysicsGeometryCollection_LLImmediate    FPhysicsGeometryCollection;
-typedef FPhysicsMaterialHandle_LLImmediate        FPhysicsMaterialHandle;
-
-
-
 #elif WITH_CHAOS
 
+using FPhysicsActorHandle = Chaos::TGeometryParticle<float, 3>*;
+
 class FChaosSceneId;
 class FPhysInterface_Chaos;
-class FPhysicsActorReference_Chaos;
 class FPhysicsConstraintReference_Chaos;
 class FPhysicsAggregateReference_Chaos;
 class FPhysicsShapeReference_Chaos;
 class FPhysScene_ChaosInterface;
+class FPhysicsShapeAdapter_Chaos;
+struct FPhysicsGeometryCollection_Chaos;
+class FPhysicsUserData_Chaos;
 
-typedef FPhysicsActorReference_Chaos		FPhysicsActorHandle;
 typedef FPhysicsConstraintReference_Chaos	FPhysicsConstraintHandle;
-typedef FPhysInterface_Chaos    		    FPhysicsInterface;
-typedef FPhysScene_ChaosInterface		    FPhysScene;
+typedef FPhysInterface_Chaos				FPhysicsInterface;
+typedef FPhysScene_ChaosInterface			FPhysScene;
 typedef FPhysicsAggregateReference_Chaos	FPhysicsAggregateHandle;
 typedef FPhysInterface_Chaos				FPhysicsCommand;
 typedef FPhysicsShapeReference_Chaos		FPhysicsShapeHandle;
-typedef FPhysicsShapeReference_Chaos	    FPhysicsGeometryCollection;
-typedef void*                           	FPhysicsMaterialHandle;
-
+typedef FPhysicsGeometryCollection_Chaos	FPhysicsGeometryCollection;
+typedef void*								FPhysicsMaterialHandle;
+typedef FPhysicsShapeAdapter_Chaos			FPhysicsShapeAdapter;
+typedef FPhysicsUserData_Chaos				FPhysicsUserData;
 
+inline FPhysicsActorHandle DefaultPhysicsActorHandle() { return nullptr; }
 
 #else
 
diff --git a/Engine/Source/Runtime/PhysicsCore/Public/PhysicsInterfaceTypesCore.h b/Engine/Source/Runtime/PhysicsCore/Public/PhysicsInterfaceTypesCore.h
--- a/Engine/Source/Runtime/PhysicsCore/Public/PhysicsInterfaceTypesCore.h
+++ b/Engine/Source/Runtime/PhysicsCore/Public/PhysicsInterfaceTypesCore.h
@@ -4,11 +4,17 @@
 #include "CoreMinimal.h"
 #include "PhysicsInterfaceDeclaresCore.h"
 
+#if INCLUDE_CHAOS
+#include "Chaos/CollisionFilterData.h"
+#endif
+
+struct FBodyInstance;
 
 struct FActorCreationParams
 {
 	FActorCreationParams()
 		: Scene(nullptr)
+		, BodyInstance(nullptr)
 		, InitialTM(FTransform::Identity)
 		, bStatic(false)
 		, bQueryOnly(false)
@@ -17,6 +23,7 @@ struct FActorCreationParams
 	{}
 
 	FPhysScene* Scene;
+	FBodyInstance* BodyInstance;
 	FTransform InitialTM;
 	bool bStatic;
 	bool bQueryOnly;
@@ -24,6 +31,7 @@ struct FActorCreationParams
 	char* DebugName;
 };
 
+#if !INCLUDE_CHAOS
 struct FCollisionFilterData
 {
 	uint32 Word0;
@@ -36,6 +44,7 @@ struct FCollisionFilterData
 		Word0 = Word1 = Word2 = Word3 = 0;
 	}
 };
+#endif
 
 /**
 * Type of query for object type or trace type
diff --git a/Engine/Source/Runtime/PhysicsCore/Public/PhysicsInterfaceUtilsCore.h b/Engine/Source/Runtime/PhysicsCore/Public/PhysicsInterfaceUtilsCore.h
--- a/Engine/Source/Runtime/PhysicsCore/Public/PhysicsInterfaceUtilsCore.h
+++ b/Engine/Source/Runtime/PhysicsCore/Public/PhysicsInterfaceUtilsCore.h
@@ -2,12 +2,13 @@
 
 #pragma once
 
+#include "PhysicsInterfaceTypesCore.h"
+#include "PhysicsInterfaceDeclaresCore.h"
+
 #if WITH_PHYSX
 
 #include "PhysXPublicCore.h"
 #include "Containers/Union.h"
-#include "PhysicsInterfaceTypesCore.h"
-#include "PhysicsInterfaceDeclaresCore.h"
 #include "CollisionShape.h"
 
 
@@ -136,14 +137,11 @@ PHYSICSCORE_API PxTransform ConvertToPhysXCapsulePose(const FTransform& GeomPose
 
 // FILTER DATA
 
-//Find the face index for a given hit. This gives us a chance to modify face index based on things like most opposing normal
-PHYSICSCORE_API PxU32 FindFaceIndex(const FHitLocation& PHit, const FVector& UnitDirection);
-
 // Adapts a FCollisionShape to a PxGeometry type, used for various queries
-struct PHYSICSCORE_API FPhysXShapeAdaptor
+struct PHYSICSCORE_API FPhysXShapeAdapter
 {
 public:
-	FPhysXShapeAdaptor(const FQuat& Rot, const FCollisionShape& CollisionShape);
+	FPhysXShapeAdapter(const FQuat& Rot, const FCollisionShape& CollisionShape);
 
 	const PxGeometry& GetGeometry() const
 	{
@@ -167,3 +165,6 @@ private:
 };
 
 #endif // WITH_PHYX
+
+//Find the face index for a given hit. This gives us a chance to modify face index based on things like most opposing normal
+PHYSICSCORE_API uint32 FindFaceIndex(const FHitLocation& PHit, const FVector& UnitDirection);
diff --git a/Engine/Source/Runtime/PhysicsCore/Public/SQCapture.h b/Engine/Source/Runtime/PhysicsCore/Public/SQCapture.h
--- a/Engine/Source/Runtime/PhysicsCore/Public/SQCapture.h
+++ b/Engine/Source/Runtime/PhysicsCore/Public/SQCapture.h
@@ -3,8 +3,6 @@
 #pragma once
 
 #include "PhysicsInterfaceDeclaresCore.h"
-#if !WITH_CHAOS_NEEDS_TO_BE_FIXED
-
 #include "CoreMinimal.h"
 #include "ChaosSQTypes.h"
 #include "PhysicsInterfaceWrapperShared.h"
@@ -12,9 +10,13 @@
 #include "PhysicsInterfaceDeclaresCore.h"
 #include "PhysXInterfaceWrapperCore.h"
 #include "CollisionQueryFilterCallbackCore.h"
+#include "ChaosInterfaceWrapperCore.h"
 
 
 class FPhysTestSerializer;
+struct FSweepHit;
+struct FRaycastHit;
+struct FOverlapHit;
 
 //Allows us to capture a scene query with either physx or chaos and then load it into either format for testing purposes
 struct PHYSICSCORE_API FSQCapture
@@ -43,15 +45,15 @@ struct PHYSICSCORE_API FSQCapture
 
 #if INCLUDE_CHAOS
 	void StartCaptureChaos(const Chaos::TImplicitObject<float, 3>& InQueryGeom, const FTransform& InStartTM, const FVector& InDir, float InDeltaMag, EHitFlags InOutputFlags);
-	void EndCaptureChaos(const FPhysicsHitCallback<FSweepHit>& Results);
+	void EndCaptureChaos(const ChaosInterface::FSQHitBuffer<ChaosInterface::FSweepHit>& Results);
 #endif
 
-#if WITH_PHYSX
-	ECollisionQueryHitType GetFilterResult(const physx::PxShape* Shape, const physx::PxActor* Actor) const;
+#if INCLUDE_CHAOS
+	ECollisionQueryHitType GetFilterResult(const Chaos::TPerShapeData<float,3>* Shape, const Chaos::TGeometryParticle<float,3>* Actor) const;
 #endif
 
-#if INCLUDE_CHAOS
-	ECollisionQueryHitType GetFilterResult(const Chaos::TImplicitObject<float,3>& Shape, const int32 ActorIdx, const int32 ShapeIdx) const;
+#if WITH_PHYSX
+	ECollisionQueryHitType GetFilterResult(const physx::PxShape* Shape, const physx::PxActor* Actor) const;
 #endif
 	
 	FVector Dir;
@@ -64,9 +66,9 @@ struct PHYSICSCORE_API FSQCapture
 	TUniquePtr<ICollisionQueryFilterCallbackBase> FilterCallback;
 
 #if WITH_PHYSX
-	FDynamicHitBuffer<PxSweepHit> PhysXSweepBuffer;
-	FDynamicHitBuffer<PxRaycastHit> PhysXRaycastBuffer;
-	FDynamicHitBuffer<PxOverlapHit> PhysXOverlapBuffer;
+	PhysXInterface::FDynamicHitBuffer<PxSweepHit> PhysXSweepBuffer;
+	PhysXInterface::FDynamicHitBuffer<PxRaycastHit> PhysXRaycastBuffer;
+	PhysXInterface::FDynamicHitBuffer<PxOverlapHit> PhysXOverlapBuffer;
 	PxGeometryHolder PhysXGeometry;
 #endif
 
@@ -75,14 +77,14 @@ struct PHYSICSCORE_API FSQCapture
 	const Chaos::TImplicitObject<float, 3>* ChaosGeometry;
 #if WITH_PHYSX
 	//for now just use physx hit buffer
-	FDynamicHitBuffer<FSweepHit> ChaosSweepBuffer;
-	TArray<FSweepHit> ChaosSweepTouches;
+	ChaosInterface::FSQHitBuffer<ChaosInterface::FSweepHit> ChaosSweepBuffer;
+	TArray<ChaosInterface::FSweepHit> ChaosSweepTouches;
 
-	FDynamicHitBuffer<FRaycastHit> ChaosRaycastBuffer;
-	TArray<FRaycastHit> ChaosRaycastTouches;
+	ChaosInterface::FSQHitBuffer<ChaosInterface::FRaycastHit> ChaosRaycastBuffer;
+	TArray<ChaosInterface::FRaycastHit> ChaosRaycastTouches;
 
-	FDynamicHitBuffer<FOverlapHit> ChaosOverlapBuffer;
-	TArray<FOverlapHit> ChaosOverlapTouches;
+	ChaosInterface::FSQHitBuffer<ChaosInterface::FOverlapHit> ChaosOverlapBuffer;
+	TArray<ChaosInterface::FOverlapHit> ChaosOverlapTouches;
 #endif
 #endif
 
@@ -122,7 +124,7 @@ private:
 	void SerializePhysXHitType(FArchive& Ar, T& Hit);
 
 	template <typename THit>
-	void SerializePhysXBuffers(FArchive& Ar, int32 Version, FDynamicHitBuffer<THit>& PhysXBuffer);
+	void SerializePhysXBuffers(FArchive& Ar, int32 Version, PhysXInterface::FDynamicHitBuffer<THit>& PhysXBuffer);
 
 	struct FPhysXSerializerData
 	{
@@ -136,7 +138,7 @@ private:
 	};
 	TUniquePtr<FPhysXSerializerData> AlignedDataHelper;
 
-	TMap<physx::PxActor*, TArray<TPair<physx::PxShape*, ECollisionQueryHitType>>> ActorToShapeHitsArray;
+	TMap<physx::PxActor*, TArray<TPair<physx::PxShape*, ECollisionQueryHitType>>> PxActorToShapeHitsArray;
 
 	//only valid during capture when serializing out runtime structures that already use non transient data
 	TMap<physx::PxActor*, physx::PxActor*> NonTransientToTransientActors;
@@ -146,11 +148,9 @@ private:
 	FPhysTestSerializer& PhysSerializer;
 
 #if INCLUDE_CHAOS
-	TMap<int32, TArray<ECollisionQueryHitType>> ChaosActorToHitsArray;
+	TMap<Chaos::TGeometryParticle<float,3>*, TArray<TPair<Chaos::TPerShapeData<float,3>*, ECollisionQueryHitType>>> ChaosActorToShapeHitsArray;
 #endif
 	bool bDiskDataIsChaos;
 	bool bChaosDataReady;
 	bool bPhysXDataReady;
 };
-
-#endif
\ No newline at end of file
diff --git a/Engine/Source/Runtime/PhysicsSQ/PhysicsSQ.Build.cs b/Engine/Source/Runtime/PhysicsSQ/PhysicsSQ.Build.cs
--- a/Engine/Source/Runtime/PhysicsSQ/PhysicsSQ.Build.cs
+++ b/Engine/Source/Runtime/PhysicsSQ/PhysicsSQ.Build.cs
@@ -9,7 +9,6 @@ public class PhysicsSQ: ModuleRules
 		PrivateDependencyModuleNames.AddRange(
 			new string[] {
 				"Core",
-				"Chaos",
 				"PhysicsCore"
 			}
 		);
diff --git a/Engine/Source/Runtime/PhysicsSQ/Private/Experimental/SceneQueryChaosImp.h b/Engine/Source/Runtime/PhysicsSQ/Private/Experimental/SceneQueryChaosImp.h
--- a/Engine/Source/Runtime/PhysicsSQ/Private/Experimental/SceneQueryChaosImp.h
+++ b/Engine/Source/Runtime/PhysicsSQ/Private/Experimental/SceneQueryChaosImp.h
@@ -4,18 +4,21 @@
 
 #if WITH_CHAOS
 
-inline bool LowLevelRaycastImp(const FVector& Start, const FVector& Dir, float DeltaMag, const Chaos::TImplicitObject<float, 3>& Shape, const FTransform ActorTM, EHitFlags OutputFlags, FPhysRaycastDummy& Hit)
+inline bool LowLevelRaycastImp(const FVector& Start, const FVector& Dir, float DeltaMag, const Chaos::TImplicitObject<float, 3>& Shape, const FTransform ActorTM, EHitFlags OutputFlags, FHitRaycast& Hit)
 {
+	//TODO_SQ_IMPLEMENTATION
 	return false;			
 }
 
-inline bool LowLevelSweepImp(const FTransform& StartTM, const FVector& Dir, float DeltaMag, const FPhysicsGeometry& SweepGeom, const Chaos::TImplicitObject<float, 3>& Shape, const FTransform ActorTM, EHitFlags OutputFlags, FPhysSweepDummy& Hit)
+inline bool LowLevelSweepImp(const FTransform& StartTM, const FVector& Dir, float DeltaMag, const FPhysicsGeometry& SweepGeom, const Chaos::TImplicitObject<float, 3>& Shape, const FTransform ActorTM, EHitFlags OutputFlags, FHitSweep& Hit)
 {
+	//TODO_SQ_IMPLEMENTATION
 	return false;
 }
 
-inline bool LowLevelOverlapImp(const FTransform& GeomPose, const FPhysicsGeometry& OverlapGeom, const Chaos::TImplicitObject<float, 3>& Shape, const FTransform ActorTM, FPhysTypeDummy& Overlap)
+inline bool LowLevelOverlapImp(const FTransform& GeomPose, const FPhysicsGeometry& OverlapGeom, const Chaos::TImplicitObject<float, 3>& Shape, const FTransform ActorTM, FHitOverlap& Overlap)
 {
+	//TODO_SQ_IMPLEMENTATION
 	return false;
 }
 #endif // WITH_CHAOS
\ No newline at end of file
diff --git a/Engine/Source/Runtime/PhysicsSQ/Private/SQAccelerator.cpp b/Engine/Source/Runtime/PhysicsSQ/Private/SQAccelerator.cpp
--- a/Engine/Source/Runtime/PhysicsSQ/Private/SQAccelerator.cpp
+++ b/Engine/Source/Runtime/PhysicsSQ/Private/SQAccelerator.cpp
@@ -6,290 +6,44 @@
 #if PHYSICS_INTERFACE_PHYSX
 #include "SceneQueryPhysXImp.h"	//todo: use nice platform wrapper
 #include "PhysXInterfaceWrapperCore.h"
-#elif PHYSICS_INTERFACE_LLIMMEDIATE
-//#include "Collision/Experimental/SceneQueryLLImmediateImp.h"
 #endif
+
 #if WITH_CHAOS
 #include "Experimental/SceneQueryChaosImp.h"
 #endif
 
+#if INCLUDE_CHAOS
+#include "ChaosInterfaceWrapperCore.h"
+
 #include "Chaos/Particles.h"
 #include "Chaos/Sphere.h"
 #include "Chaos/Capsule.h"
 #include "Chaos/PBDRigidsEvolutionGBF.h"
 #include "Chaos/PBDCollisionConstraint.h"
 #include "Chaos/GeometryQueries.h"
-
-FSQAcceleratorEntry* FSQAccelerator::AddEntry(void* Payload)
-{
-	if (Nodes.Num() == 0)
-	{
-		Nodes.Add(new FSQNode());
-	}
-
-	FSQAcceleratorEntry* NewEntry = new FSQAcceleratorEntry(Payload);
-	Nodes[0]->Entries.Add(NewEntry);
-	return NewEntry;
-}
-
-#if 0
-void FSQAccelerator::UpdateBounds(FSQAcceleratorEntry* Entry, const FBoxSphereBounds& NewBounds)
-{
-	Entry->Bounds = NewBounds;
-}
-#endif
-
-void FSQAccelerator::RemoveEntry(FSQAcceleratorEntry* Entry)
-{
-	if (Nodes.Num())
-	{
-		Nodes[0]->Entries.RemoveSingleSwap(Entry);
-	}
-	delete Entry;
-}
-
-void FSQAccelerator::GetNodes(TArray<const FSQNode*>& NodesFound) const
-{
-	NodesFound.Reset(Nodes.Num());
-	for (const FSQNode* Node : Nodes)
-	{
-		NodesFound.Add(Node);
-	}
-}
-
-FSQAccelerator::~FSQAccelerator()
-{
-	for (FSQNode* Node : Nodes)
-	{
-		for (FSQAcceleratorEntry* Entry : Node->Entries)
-		{
-			delete Entry;
-		}
-
-		delete Node;
-	}
-}
-
-template <typename HitType>
-class FSQOverlapBuffer
-{
-public:
-
-	FSQOverlapBuffer(int32 InMaxNumOverlaps)
-		: MaxNumOverlaps(InMaxNumOverlaps)
-	{
-	}
-
-	bool Insert(const HitType& Hit)
-	{
-		//maybe prioritize based on penetration depth?
-		if (Overlapping.Num() < MaxNumOverlaps)
-		{
-			Overlapping.Add(Hit);
-		}
-
-		return Overlapping.Num() < MaxNumOverlaps;
-	}
-
-private:
-	TArray<HitType>	Overlapping;	//todo(ocohen): use a TInlineArray + compatible bytes to avoid default constructor and allocations
-	int32 MaxNumOverlaps;
-};
-
-template <typename HitType>
-class FSQTraceBuffer
-{
-public:
-
-	FSQTraceBuffer(float InDeltaMag, int32 InMaxNumOverlaps)
-		: DeltaMag(InDeltaMag)
-		, MaxNumOverlaps(InMaxNumOverlaps)
-	{
-	}
-
-	bool Insert(const HitType& Hit, bool bBlocking)
-	{
-		if (GetDistance(Hit) < DeltaMag)
-		{
-			if (bBlocking)
-			{
-				BlockingHit = Hit;
-				bHasBlocking = true;
-				DeltaMag = GetDistance(Hit);
-			}
-			else
-			{
-				//todo(ocohen):use a priority queue or something more sensible than just a full sort every time. Avoid add then remove if possible
-				Overlapping.Add(Hit);
-				Overlapping.Sort();
-				if (Overlapping.Num() > MaxNumOverlaps)
-				{
-					Overlapping.SetNum(MaxNumOverlaps);
-				}
-			}
-		}
-
-		return true;
-	}
-
-	float GetBlockingDistance() const { return DeltaMag; }
-	float GetOverlapingDistance() const { return DeltaMag;  /*todo(ocohen): if we keep overlaps sorted we could make this tighter and avoid memory growth.*/ }
-
-private:
-	HitType BlockingHit;
-	TArray<HitType>	Overlapping;	//todo(ocohen): use a TInlineArray + compatible bytes to avoid default constructor and allocations
-	int32 MaxNumOverlaps;
-	bool bHasBlocking;
-	float DeltaMag;
-};
-
-void FSQAccelerator::Raycast(const FVector& Start, const FVector& Dir, const float DeltaMagnitude, FPhysicsHitCallback<FHitRaycast>& HitBuffer, EHitFlags OutputFlags, FQueryFlags QueryFlags, const FCollisionFilterData& QueryFilter, const FQueryFilterData& QueryFilterData, ICollisionQueryFilterCallbackBase& QueryCallback) const
-{
-#if WITH_PHYSX
-	FPhysicsRaycastInputAdapater Inputs(Start, Dir, OutputFlags);
-
-	for (const FSQNode* Node : Nodes)
-	{
-		//in theory we'd do some checks against max block and max overlap distance to avoid some of these nodes
-		for (const FSQAcceleratorEntry* Entry : Node->Entries)
-		{
-			if (FPhysicsActor* RigidActor = (FPhysicsActor*)Entry->GetPayload())
-			{
-				const FPhysicsTransform ActorTM = GetGlobalPose(*RigidActor);
-				const uint32 NumShapes = GetNumShapes(*RigidActor);
-				TArray<FPhysicsShape*, TInlineAllocator<16>> Shapes;
-				Shapes.AddZeroed(NumShapes);
-				GetShapes(*RigidActor, Shapes.GetData(), NumShapes);
-				FHitRaycast Hit;
-				for (FPhysicsShape* Shape : Shapes)
-				{
-					const FCollisionFilterData ShapeFilterData = GetQueryFilterData(*Shape);
-					ECollisionQueryHitType FilterType = QueryFlags & EQueryFlags::PreFilter ? QueryCallback.PreFilter(QueryFilter, *Shape, *RigidActor) : ECollisionQueryHitType::Block;	//todo(ocohen): we always use preFilter, should we add a cheaper test?
-					if (FilterType != ECollisionQueryHitType::None)
-					{
-						if (LowLevelRaycastImp(Inputs.Start, Inputs.Dir, DeltaMagnitude, *Shape, ActorTM, Inputs.OutputFlags, Hit))
-						{
-							SetActor(Hit, RigidActor);
-							SetShape(Hit, Shape);
-							FilterType = QueryFlags & EQueryFlags::PostFilter ? QueryCallback.PostFilter(QueryFilter, Hit) : FilterType;
-							if (FilterType != ECollisionQueryHitType::None)
-							{
-								Insert(HitBuffer, Hit, FilterType == ECollisionQueryHitType::Block || (QueryFlags & EQueryFlags::AnyHit));
-							}
-						}
-					}
-				}
-			}
-		}	
-	}
 #endif
-}
 
-void FSQAccelerator::Sweep(const FPhysicsGeometry& QueryGeom, const FTransform& StartTM, const FVector& Dir, const float DeltaMagnitude, FPhysicsHitCallback<FHitSweep>& HitBuffer, EHitFlags OutputFlags, FQueryFlags QueryFlags, const FCollisionFilterData& QueryFilter, const FQueryFilterData& QueryFilterData, ICollisionQueryFilterCallbackBase& QueryCallback) const
-{
-#if WITH_PHYSX
-	FPhysicsSweepInputAdapater Inputs(StartTM, Dir, OutputFlags);
-	for (const FSQNode* Node : Nodes)
-	{
-		for (const FSQAcceleratorEntry* Entry : Node->Entries)
-		{
-			if (FPhysicsActor* RigidActor = (FPhysicsActor*)Entry->GetPayload())
-			{
-				const FPhysicsTransform ActorTM = GetGlobalPose(*RigidActor);
-				const uint32 NumShapes = GetNumShapes(*RigidActor);
-				TArray<FPhysicsShape*, TInlineAllocator<16>> Shapes;
-				Shapes.AddZeroed(NumShapes);
-				GetShapes(*RigidActor, Shapes.GetData(), NumShapes);
-				FHitSweep Hit;
-				for (FPhysicsShape* Shape : Shapes)
-				{
-					const FCollisionFilterData ShapeFilterData = GetQueryFilterData(*Shape);
-					ECollisionQueryHitType FilterType = QueryFlags & EQueryFlags::PreFilter ? QueryCallback.PreFilter(QueryFilter, *Shape, *RigidActor) : ECollisionQueryHitType::Block;	//todo(ocohen): we always use preFilter, should we add a cheaper test?
-					if (FilterType != ECollisionQueryHitType::None)
-					{
-						if (LowLevelSweepImp(Inputs.StartTM, Inputs.Dir, DeltaMagnitude, QueryGeom, *Shape, ActorTM, Inputs.OutputFlags, Hit))
-						{
-							SetActor(Hit, RigidActor);
-							SetShape(Hit, Shape);
-							FilterType = QueryFlags & EQueryFlags::PostFilter ? QueryCallback.PostFilter(QueryFilter, Hit) : FilterType;
-							if (FilterType != ECollisionQueryHitType::None)
-							{
-								Insert(HitBuffer, Hit, FilterType == ECollisionQueryHitType::Block || (QueryFlags & EQueryFlags::AnyHit));
-							}
-						}
-					}
-
-				}
-			}
-		}
-	}
-#endif
-}
-
-void FSQAccelerator::Overlap(const FPhysicsGeometry& QueryGeom, const FTransform& GeomPose, FPhysicsHitCallback<FHitOverlap>& HitBuffer, FQueryFlags QueryFlags, const FCollisionFilterData& QueryFilter, const FQueryFilterData& QueryFilterData, ICollisionQueryFilterCallbackBase& QueryCallback) const
-{
-#if WITH_PHYSX
-	FPhysicsOverlapInputAdapater Inputs(GeomPose);
-	for (const FSQNode* Node : Nodes)
-	{
-		for (const FSQAcceleratorEntry* Entry : Node->Entries)
-		{
-			if (FPhysicsActor* RigidActor = (FPhysicsActor*)Entry->GetPayload())
-			{
-				const FPhysicsTransform ActorTM = GetGlobalPose(*RigidActor);
-				const uint32 NumShapes = GetNumShapes(*RigidActor);
-				TArray<FPhysicsShape*, TInlineAllocator<16>> Shapes;
-				Shapes.AddZeroed(NumShapes);
-				GetShapes(*RigidActor, Shapes.GetData(), NumShapes);
-				FHitOverlap Overlap;
-				for (FPhysicsShape* Shape : Shapes)
-				{
-					const FCollisionFilterData ShapeFilterData = GetQueryFilterData(*Shape);
-					ECollisionQueryHitType FilterType = QueryFlags & EQueryFlags::PreFilter ? QueryCallback.PreFilter(QueryFilter, *Shape, *RigidActor) : ECollisionQueryHitType::Block;	//todo(ocohen): we always use preFilter, should we add a cheaper test?
-					if (FilterType != ECollisionQueryHitType::None)
-					{
-						if (LowLevelOverlapImp(Inputs.GeomPose, QueryGeom, *Shape, ActorTM, Overlap))
-						{
-							SetActor(Overlap, RigidActor);
-							SetShape(Overlap, Shape);
-							FilterType = QueryFlags & EQueryFlags::PostFilter ? QueryCallback.PostFilter(QueryFilter, Overlap) : FilterType;
-							if (FilterType != ECollisionQueryHitType::None)
-							{
-								if (!InsertOverlap(HitBuffer, Overlap))
-								{
-									return;
-								}
-							}
-						}
-					}
-				}
-			}
-		}
-	}
-#endif
-}
-
-void FSQAcceleratorUnion::Raycast(const FVector& Start, const FVector& Dir, const float DeltaMagnitude, FPhysicsHitCallback<FHitRaycast>& HitBuffer, EHitFlags OutputFlags, FQueryFlags QueryFlags, const FCollisionFilterData& QueryFilter, const FQueryFilterData& QueryFilterData, ICollisionQueryFilterCallbackBase& QueryCallback) const
+void FSQAcceleratorUnion::Raycast(const FVector& Start, const FVector& Dir, const float DeltaMagnitude, FPhysicsHitCallback<FHitRaycast>& HitBuffer, EHitFlags OutputFlags, const FQueryFilterData& QueryFilterData, ICollisionQueryFilterCallbackBase& QueryCallback) const
 {
 	for (const ISQAccelerator* Accelerator : Accelerators)
 	{
-		Accelerator->Raycast(Start, Dir, DeltaMagnitude, HitBuffer, OutputFlags, QueryFlags, QueryFilter, QueryFilterData, QueryCallback);
+		Accelerator->Raycast(Start, Dir, DeltaMagnitude, HitBuffer, OutputFlags, QueryFilterData, QueryCallback);
 	}
 }
 
-void FSQAcceleratorUnion::Sweep(const FPhysicsGeometry& QueryGeom, const FTransform& StartTM, const FVector& Dir, const float DeltaMagnitude, FPhysicsHitCallback<FHitSweep>& HitBuffer, EHitFlags OutputFlags, FQueryFlags QueryFlags, const FCollisionFilterData& QueryFilter, const FQueryFilterData& QueryFilterData, ICollisionQueryFilterCallbackBase& QueryCallback) const
+void FSQAcceleratorUnion::Sweep(const FPhysicsGeometry& QueryGeom, const FTransform& StartTM, const FVector& Dir, const float DeltaMagnitude, FPhysicsHitCallback<FHitSweep>& HitBuffer, EHitFlags OutputFlags, const FQueryFilterData& QueryFilterData, ICollisionQueryFilterCallbackBase& QueryCallback) const
 {
 	for (const ISQAccelerator* Accelerator : Accelerators)
 	{
-		Accelerator->Sweep(QueryGeom, StartTM, Dir, DeltaMagnitude, HitBuffer, OutputFlags, QueryFlags, QueryFilter, QueryFilterData, QueryCallback);
+		Accelerator->Sweep(QueryGeom, StartTM, Dir, DeltaMagnitude, HitBuffer, OutputFlags, QueryFilterData, QueryCallback);
 	}
 }
 
-void FSQAcceleratorUnion::Overlap(const FPhysicsGeometry& QueryGeom, const FTransform& GeomPose, FPhysicsHitCallback<FHitOverlap>& HitBuffer, FQueryFlags QueryFlags, const FCollisionFilterData& QueryFilter, const FQueryFilterData& QueryFilterData, ICollisionQueryFilterCallbackBase& QueryCallback) const
+void FSQAcceleratorUnion::Overlap(const FPhysicsGeometry& QueryGeom, const FTransform& GeomPose, FPhysicsHitCallback<FHitOverlap>& HitBuffer, const FQueryFilterData& QueryFilterData, ICollisionQueryFilterCallbackBase& QueryCallback) const
 {
 	for (const ISQAccelerator* Accelerator : Accelerators)
 	{
-		Accelerator->Overlap(QueryGeom, GeomPose, HitBuffer, QueryFlags, QueryFilter, QueryFilterData, QueryCallback);
+		Accelerator->Overlap(QueryGeom, GeomPose, HitBuffer, QueryFilterData, QueryCallback);
 	}
 }
 
@@ -314,352 +68,267 @@ struct FPreFilterInfo
 	int32 ActorIdx;
 };
 
-void GeneratePrefilterInfo(const Chaos::TPBDRigidsEvolutionGBF<float, 3>& Evolution, const TArray<int32>& PotentialIntersections, const FCollisionFilterData& QueryFilter, const FQueryFilterData& QueryFilterData, bool bAnyHit, ICollisionQueryFilterCallbackBase& QueryCallback, TArray<FPreFilterInfo>& OutTouchers, TArray<FPreFilterInfo>& OutBlockers)
+void FillHitHelper(ChaosInterface::FLocationHit& Hit, const float Distance, const FVector& WorldPosition, const FVector& WorldNormal)
 {
-#if WITH_PHYSX
-	using namespace Chaos;
-	for (int32 PotentialIdx : PotentialIntersections)
-	{
-		const TPBDRigidParticles<float, 3>& Particles = Evolution.GetParticles();
-		const TImplicitObject<float, 3>* Object = Particles.Geometry(PotentialIdx).Get();
-		if (!Object)
-		{
-			continue;
-		}
-
-		TArray<const TImplicitObject<float, 3>*> Geoms;
-		if (const TImplicitObjectUnion<float, 3>* UnionObject = Object->template GetObject<const TImplicitObjectUnion<float, 3>>())
-		{
-			const TArray <TUniquePtr<TImplicitObject<float, 3>>>& UnionGeoms = UnionObject->GetObjects();
-			Geoms.Reserve(UnionGeoms.Num());
-			for (const auto& GeomPtr : UnionGeoms)
-			{
-				Geoms.Add(GeomPtr.Get());
-			}
-		}
-		else
-		{
-			Geoms.Add(Object);
-		}
-
-		for (int32 GeomIdx = 0; GeomIdx < Geoms.Num(); ++GeomIdx)
-		{
-			const TImplicitObject<float, 3>* Geom = Geoms[GeomIdx];
-			const ECollisionQueryHitType HitType = QueryFilterData.flags & PxQueryFlag::ePREFILTER ? QueryCallback.PreFilterChaos(QueryFilter, *Geom, PotentialIdx, GeomIdx) : ECollisionQueryHitType::Block;
-			if (HitType == ECollisionQueryHitType::Block)
-			{
-				OutBlockers.Add({ Geom, PotentialIdx });
-			}
-			else if (HitType == ECollisionQueryHitType::Touch)
-			{
-				if (bAnyHit)
-				{
-					OutBlockers.Add({ Geom, PotentialIdx });
-				}
-				else
-				{
-					OutTouchers.Add({ Geom, PotentialIdx });
-				}
-			}
-		}
-	}
-#endif
+	Hit.Distance = Distance;
+	Hit.WorldPosition = WorldPosition;
+	Hit.WorldNormal = WorldNormal;
+	Hit.Flags = Distance > 0 ? EHitFlags::Distance | EHitFlags::Normal | EHitFlags::Position : EHitFlags::Distance;
 }
 
-bool ProcessRaycasts(const Chaos::TPBDRigidsEvolutionGBF<float, 3>& Evolution, const TArray<FPreFilterInfo>& PreFilters, const FVector& StartPoint, const FVector& Dir, const float DeltaMag, FPhysicsHitCallback<FRaycastHit>& HitBuffer, EHitFlags OutputFlags, ECollisionQueryHitType HitType, bool bAnyHit)
+void FillHitHelper(ChaosInterface::FOverlapHit& Hit, const float Distance, const FVector& WorldPosition, const FVector& WorldNormal)
 {
-	using namespace Chaos;
-	const TPBDRigidParticles<float, 3>& Particles = Evolution.GetParticles();
-	for (const FPreFilterInfo& PreFilter : PreFilters)
-	{
-		const TRigidTransform<float, 3> TM(Particles.X(PreFilter.ActorIdx), Particles.R(PreFilter.ActorIdx));
-		const TVector<float, 3> DirLocal = TM.InverseTransformVectorNoScale(Dir);
-
-		FRaycastHit Hit;
-		const TVector<float, 3> StartLocal = TM.InverseTransformPositionNoScale(StartPoint);
-
-		TVector<float, 3> LocalNormal;
-		TVector<float, 3> LocalPosition;
-		float Time;
-
-		if (PreFilter.Geom->Raycast(StartLocal, DirLocal, DeltaMag, /*Thickness=*/0.f, Time, LocalPosition, LocalNormal))
-		{
-			Hit.Distance = DeltaMag * Time;
-			Hit.WorldPosition = TM.TransformPositionNoScale(LocalPosition);
-			Hit.WorldNormal = TM.TransformVectorNoScale(LocalNormal);
-			Hit.Flags = EHitFlags::Distance | EHitFlags::Normal | EHitFlags::Position;
-
-			Hit.ActorIdx = PreFilter.ActorIdx;
-			Hit.Shape = PreFilter.Geom;
-			Insert(HitBuffer, Hit, HitType == ECollisionQueryHitType::Block);
-			if (bAnyHit)
-			{
-				return false;
-			}
-		}
-	}
-
-	return true;
 }
 
-
-void FChaosSQAccelerator::ChaosRaycast(const FVector& StartPoint, const FVector& Dir, const float DeltaMag, FPhysicsHitCallback<FRaycastHit>& HitBuffer, EHitFlags OutputFlags, const FCollisionFilterData& QueryFilter, const FQueryFilterData& QueryFilterData, ICollisionQueryFilterCallbackBase& QueryCallback) const
+template <typename TPayload, typename THitType>
+struct TSQVisitor : public Chaos::ISpatialVisitor<TPayload, float>
 {
-#if WITH_PHYSX
-	using namespace Chaos;
-
-	TArray<int32> PotentialIntersections;
+	TSQVisitor(const FVector& InStartPoint, const FVector& InDir, ChaosInterface::FSQHitBuffer<ChaosInterface::FRaycastHit>& InHitBuffer, EHitFlags InOutputFlags,
+		const FQueryFilterData& InQueryFilterData, ICollisionQueryFilterCallbackBase& InQueryCallback)
+		: StartPoint(InStartPoint)
+		, Dir(InDir)
+		, HitBuffer(InHitBuffer)
+		, OutputFlags(InOutputFlags)
+		, QueryFilterData(InQueryFilterData)
+		, QueryCallback(InQueryCallback)
+		, bAnyHit(false)
 	{
-		const ISpatialAcceleration<float, 3>& SpatialAcceleration = Evolution.GetSpatialAcceleration();
-		const TVector<float, 3>& RayStart = StartPoint;
-		const TVector<float, 3> RayEnd = RayStart + Dir * DeltaMag;
-		PotentialIntersections = SpatialAcceleration.FindAllIntersections(TSpatialRay<float, 3>(RayStart, RayEnd));
-		Evolution.ReleaseSpatialAcceleration();
+#if WITH_PHYSX
+		//#TODO - reimplement query flags alternative for Chaos
+		bAnyHit = QueryFilterData.flags & PxQueryFlag::eANY_HIT;
+#endif
 	}
 
-	TArray<FPreFilterInfo> Touchers;
-	TArray<FPreFilterInfo> Blockers;
-	const bool bAnyHit = QueryFilterData.flags & PxQueryFlag::eANY_HIT;
-	GeneratePrefilterInfo(Evolution, PotentialIntersections, QueryFilter, QueryFilterData, bAnyHit, QueryCallback, Touchers, Blockers);
-
-	const bool bProcessTouches = ProcessRaycasts(Evolution, Blockers, StartPoint, Dir, DeltaMag, HitBuffer, OutputFlags, ECollisionQueryHitType::Block, bAnyHit);
-	if (bProcessTouches)
+	TSQVisitor(const FTransform& InStartTM, const FVector& InDir, ChaosInterface::FSQHitBuffer<ChaosInterface::FSweepHit>& InHitBuffer, EHitFlags InOutputFlags,
+		const FQueryFilterData& InQueryFilterData, ICollisionQueryFilterCallbackBase& InQueryCallback, const Chaos::TImplicitObject<float,3>& InQueryGeom)
+		: StartTM(InStartTM)
+		, Dir(InDir)
+		, HitBuffer(InHitBuffer)
+		, OutputFlags(InOutputFlags)
+		, QueryFilterData(InQueryFilterData)
+		, QueryCallback(InQueryCallback)
+		, bAnyHit(false)
+		, QueryGeom(&InQueryGeom)
 	{
-		ProcessRaycasts(Evolution, Touchers, StartPoint, Dir, DeltaMag, HitBuffer, OutputFlags, ECollisionQueryHitType::Touch, bAnyHit);
-	}
-
-	FinalizeQuery(HitBuffer);
+#if WITH_PHYSX
+		//#TODO - reimplement query flags alternative for Chaos
+		bAnyHit = QueryFilterData.flags & PxQueryFlag::eANY_HIT;
 #endif
-}
+		//todo: check THitType is sweep
+	}
 
-Chaos::TParticles<float, 3> GenerateSamplePoints(const Chaos::TImplicitObject<float, 3>& QueryGeom, float& OutThickness)
-{
-	using namespace Chaos;
-	TParticles<float, 3> SamplePoints;
-	OutThickness = 0.f;
-	switch (QueryGeom.GetType())
+	TSQVisitor(const FTransform& InWorldTM, ChaosInterface::FSQHitBuffer<ChaosInterface::FOverlapHit>& InHitBuffer,
+		const FQueryFilterData& InQueryFilterData, ICollisionQueryFilterCallbackBase& InQueryCallback, const Chaos::TImplicitObject<float, 3>& InQueryGeom)
+		: StartTM(InWorldTM)
+		, HitBuffer(InHitBuffer)
+		, QueryFilterData(InQueryFilterData)
+		, QueryCallback(InQueryCallback)
+		, bAnyHit(false)
+		, QueryGeom(&InQueryGeom)
 	{
-	case ImplicitObjectType::Sphere:
-	{
-		const TSphere<float, 3>& Sphere = static_cast<const TSphere<float, 3>&>(QueryGeom);
-		const float Radius = Sphere.Radius();
-		OutThickness = Radius;
-
-		SamplePoints.AddParticles(1);
-		SamplePoints.X(0) = Chaos::TVector<float, 3>(0);
-		break;
+#if WITH_PHYSX
+		//#TODO - reimplement query flags alternative for Chaos
+		bAnyHit = QueryFilterData.flags & PxQueryFlag::eANY_HIT;
+#endif
+		//todo: check THitType is overlap
 	}
-	case ImplicitObjectType::Box:
+
+	virtual bool Raycast(const TPayload Payload, float& CurLength) override
 	{
-		const TBox<float, 3>& Box = static_cast<const TBox<float, 3>&>(QueryGeom);
-		Chaos::TVector<float, 3> X1(Box.Extents() * -0.5f);
-		Chaos::TVector<float, 3> X2(-X1);
-
-		SamplePoints.AddParticles(8);
-		SamplePoints.X(0) = TVector<float, 3>(X1.X, X1.Y, X1.Z);
-		SamplePoints.X(1) = TVector<float, 3>(X1.X, X1.Y, X2.Z);
-		SamplePoints.X(2) = TVector<float, 3>(X1.X, X2.Y, X1.Z);
-		SamplePoints.X(3) = TVector<float, 3>(X2.X, X1.Y, X1.Z);
-		SamplePoints.X(4) = TVector<float, 3>(X2.X, X2.Y, X2.Z);
-		SamplePoints.X(5) = TVector<float, 3>(X2.X, X2.Y, X1.Z);
-		SamplePoints.X(6) = TVector<float, 3>(X2.X, X1.Y, X2.Z);
-		SamplePoints.X(7) = TVector<float, 3>(X1.X, X2.Y, X2.Z);
-		break;
+		return Visit<ESQType::Raycast>(Payload, CurLength);
 	}
-	case ImplicitObjectType::Capsule:
+
+	virtual bool Sweep(const TPayload Payload, float& CurLength) override
 	{
-		const TCapsule<float>& Capsule = static_cast<const TCapsule<float>&>(QueryGeom);
-		const float Radius = Capsule.GetRadius();
-		OutThickness = Radius;
-		const float HalfHeight = Capsule.GetHeight() * 0.5f;
-
-		//to at least be accurate, we add all the points such that the capsule is made up of many spheres which we can currently raycast correctly
-		const int32 NumPtsMinusOne = 3 + (HalfHeight / Radius) * 2.f;
-		SamplePoints.AddParticles(NumPtsMinusOne);
-		for (int32 Idx = 0; Idx < NumPtsMinusOne; ++Idx)
-		{
-			SamplePoints.X(Idx) = TVector<float, 3>(HalfHeight + Radius - Radius * Idx, 0, 0);
-		}
-		SamplePoints.AddParticles(1);
-		SamplePoints.X(NumPtsMinusOne) = TVector<float, 3>(-HalfHeight - Radius, 0, 0);	//if last sphere didn't go all the way, add one
-
-		/*SamplePoints.AddParticles(14);
-		SamplePoints.X(0) = TVector<float, 3>(HalfHeight + Radius, 0, 0);
-		SamplePoints.X(1) = TVector<float, 3>(-HalfHeight - Radius, 0, 0);
-		SamplePoints.X(2) = TVector<float, 3>(HalfHeight, Radius, Radius);
-		SamplePoints.X(3) = TVector<float, 3>(HalfHeight, -Radius, Radius);
-		SamplePoints.X(4) = TVector<float, 3>(HalfHeight, -Radius, -Radius);
-		SamplePoints.X(5) = TVector<float, 3>(HalfHeight, Radius, -Radius);
-		SamplePoints.X(6) = TVector<float, 3>(0, Radius, Radius);
-		SamplePoints.X(7) = TVector<float, 3>(0, -Radius, Radius);
-		SamplePoints.X(8) = TVector<float, 3>(0, -Radius, -Radius);
-		SamplePoints.X(9) = TVector<float, 3>(0, Radius, -Radius);
-		SamplePoints.X(10) = TVector<float, 3>(-HalfHeight, Radius, Radius);
-		SamplePoints.X(11) = TVector<float, 3>(-HalfHeight, -Radius, Radius);
-		SamplePoints.X(12) = TVector<float, 3>(-HalfHeight, -Radius, -Radius);
-		SamplePoints.X(13) = TVector<float, 3>(-HalfHeight, Radius, -Radius);*/
-		break;
+		return Visit<ESQType::Sweep>(Payload, CurLength);
 	}
-	default:
+
+	virtual bool Overlap(const TPayload Payload) override
 	{
-		ensureMsgf(false, TEXT("Unsupported sweep geometry"));
-	}
+		float DummyLength;
+		return Visit<ESQType::Overlap>(Payload, DummyLength);
 	}
 
-	return SamplePoints;
-}
+private:
 
-bool ProcessSweeps(const Chaos::TPBDRigidsEvolutionGBF<float, 3>& Evolution, const TArray<FPreFilterInfo>& PreFilters, const Chaos::TImplicitObject<float,3>& QueryGeom, const FTransform& StartTM, const FVector& Dir, const float Length, FPhysicsHitCallback<FSweepHit>& HitBuffer, EHitFlags OutputFlags, ECollisionQueryHitType HitType, bool bAnyHit)
-{
-	using namespace Chaos;
-	const TPBDRigidParticles<float, 3>& Particles = Evolution.GetParticles();
-	for (const FPreFilterInfo& PreFilter : PreFilters)
+	enum class ESQType
 	{
-		const TRigidTransform<float, 3> TM(Particles.X(PreFilter.ActorIdx), Particles.R(PreFilter.ActorIdx));
+		Raycast,
+		Sweep,
+		Overlap
+	};
 
-		FSweepHit Hit;
-		bool bFound = false;
-		{
-			TVector<float, 3> Normal;
-			TVector<float, 3> Position;
-			float Time;
+	template <ESQType SQ>
+	bool Visit(const TPayload Payload, float& CurLength)
+	{
+		//todo: add a check to ensure hitbuffer matches SQ type
+		using namespace Chaos;
+		const TShapesArray<float,3>& Shapes = Payload->GTGeometryParticle()->ShapesArray();	//TODO: use GT particles directly
 
-			if (SweepQuery<float, 3>(*PreFilter.Geom, TM, QueryGeom, StartTM, Dir, Length, Time, Position, Normal))
+		for (const auto& Shape : Shapes)
+		{
+			const TImplicitObject<float, 3>* Geom = Shape->Geometry;
+			//TODO: use gt particles directly
+			//#TODO alternative to px flags
+#if WITH_PHYSX
+			ECollisionQueryHitType HitType = QueryFilterData.flags & PxQueryFlag::ePREFILTER ? QueryCallback.PreFilter(P2UFilterData(QueryFilterData.data), *Shape, *Payload->GTGeometryParticle()) : ECollisionQueryHitType::Block;
+#else
+			//#TODO Chaos flag alternative
+			ensure(false);
+			ECollisionQueryHitType HitType = ECollisionQueryHitType::Block;
+#endif
+			if (HitType != ECollisionQueryHitType::None)
 			{
-				const float Distance = Time * Length;	//assumes distance is the distance traveled by the sweeping geometry, not the distance including the point of intersection
-				if (!bFound || Distance < Hit.Distance)
+				const TRigidTransform<float, 3> ActorTM(Payload->X(), Payload->R());
+
+				THitType Hit;
+				Hit.Actor = Payload->GTGeometryParticle();	//todo: use GT particles directly
+				Hit.Shape = Shape.Get();
+
+				bool bHit = false;
+
+				TVector<float, 3> WorldPosition, WorldNormal;
+				float Distance = 0;	//not needed but fixes compiler warning for overlap
+				int32 FaceIdx;	//todo: pass back to unreal
+
+				if (SQ == ESQType::Raycast)
 				{
-					Hit.Distance = Distance;
-					Hit.WorldPosition = Position;
-					Hit.WorldNormal = Normal;
-					Hit.Flags = EHitFlags::Distance | EHitFlags::Normal | EHitFlags::Position;
-					bFound = true;
-					if (bAnyHit)
+					TVector<float, 3> LocalNormal;
+					TVector<float, 3> LocalPosition;
+
+					const TVector<float, 3> DirLocal = ActorTM.InverseTransformVectorNoScale(Dir);
+					const TVector<float, 3> StartLocal = ActorTM.InverseTransformPositionNoScale(StartPoint);
+					bHit = Geom->Raycast(StartLocal, DirLocal, CurLength, /*Thickness=*/0.f, Distance, LocalPosition, LocalNormal, FaceIdx);
+					if (bHit)
 					{
-						break;
+						WorldPosition = ActorTM.TransformPositionNoScale(LocalPosition);
+						WorldNormal = ActorTM.TransformVectorNoScale(LocalNormal);
 					}
 				}
-			}
-		}
-
-		if (bFound)
-		{
-			Hit.ActorIdx = PreFilter.ActorIdx;
-			Hit.Shape = PreFilter.Geom;
-			Insert(HitBuffer, Hit, HitType == ECollisionQueryHitType::Block);
-			if (bAnyHit)
-			{
-				return false;
-			}
-		}
-	}
-
-	return true;
-}
+				else if(SQ == ESQType::Sweep && CurLength > 0)
+				{
+					bHit = SweepQuery<float, 3>(*Geom, ActorTM, *QueryGeom, StartTM, Dir, CurLength, Distance, WorldPosition, WorldNormal, FaceIdx);
+				}
+				else if (SQ == ESQType::Overlap || (SQ == ESQType::Sweep && CurLength == 0))
+				{
+					bHit = OverlapQuery<float, 3>(*Geom, ActorTM, *QueryGeom, StartTM, /*Thickness=*/0);
+				}
 
-void FChaosSQAccelerator::ChaosSweep(const Chaos::TImplicitObject<float, 3>& QueryGeom, const FTransform& StartTM, const FVector& Dir, const float DeltaMag, FPhysicsHitCallback<FSweepHit>& HitBuffer, EHitFlags OutputFlags, const FCollisionFilterData& QueryFilter, const FQueryFilterData& QueryFilterData, ICollisionQueryFilterCallbackBase& QueryCallback) const
-{
+				if(bHit)
+				{
+					FillHitHelper(Hit, Distance, WorldPosition, WorldNormal);
 #if WITH_PHYSX
-	using namespace Chaos;
+					HitType = QueryFilterData.flags & PxQueryFlag::ePOSTFILTER ? QueryCallback.PostFilter(P2UFilterData(QueryFilterData.data), Hit) : HitType;
+#else
+					//#TODO Chaos flag alternative
+					ensure(false);
+#endif
+					if (HitType != ECollisionQueryHitType::None)
+					{
 
-	//todo: handle levelset
-	//todo: handle sphere/bounds sweep through spatial acceleration
-	float Thickness;
-	TParticles<float, 3> SamplePoints = GenerateSamplePoints(QueryGeom, Thickness);
-	
-	TArray<int32> PotentialIntersections;
-	{
-		const ISpatialAcceleration<float, 3>& SpatialAcceleration = Evolution.GetSpatialAcceleration();
-		const int32 NumCollisionParticles = SamplePoints.Size();
-		for (int32 ParticleIndex = 0; ParticleIndex < NumCollisionParticles; ++ParticleIndex)
-		{
-			const TVector<float, 3> RayStart = StartTM.TransformPositionNoScale(SamplePoints.X(ParticleIndex));
-			const TVector<float, 3> RayEnd = RayStart + Dir * DeltaMag;
+						//overlap never blocks
+						const bool bBlocker = SQ != ESQType::Overlap && (HitType == ECollisionQueryHitType::Block || bAnyHit || HitBuffer.WantsSingleResult());
+						HitBuffer.InsertHit(Hit, bBlocker);
 
-			PotentialIntersections.Append(SpatialAcceleration.FindAllIntersections(TSpatialRay<float, 3>(RayStart, RayEnd)));
-		}
+						if (bBlocker)
+						{
+							CurLength = Distance;
+							if (CurLength == 0 && (SQ == ESQType::Raycast || HitBuffer.WantsSingleResult()))	//raycasts always fail with distance 0, sweeps only matter if we want multi overlaps
+							{
+								return false;	//initial overlap so nothing will be better than this
+							}
+						}
 
-		Evolution.ReleaseSpatialAcceleration();
-	}
-	PotentialIntersections.Sort();	//remove duplicates. Todo: don't use multiple sweeps in broadphase
-	int32 PrevInstance = INDEX_NONE;
-	for (int32 Idx = PotentialIntersections.Num() - 1; Idx >= 0; --Idx)
-	{
-		int32 Instance = PotentialIntersections[Idx];
-		if (Instance == PrevInstance)
-		{
-			PotentialIntersections.RemoveAtSwap(Idx);
+						if (bAnyHit)
+						{
+							return false;
+						}
+					}
+				}
+			}
 		}
-		PrevInstance = Instance;
-	}
-
-	TArray<FPreFilterInfo> Touchers;
-	TArray<FPreFilterInfo> Blockers;
-	const bool bAnyHit = QueryFilterData.flags & PxQueryFlag::eANY_HIT;
-	GeneratePrefilterInfo(Evolution, PotentialIntersections, QueryFilter, QueryFilterData, bAnyHit, QueryCallback, Touchers, Blockers);
 
-	const bool bProcessTouches = ProcessSweeps(Evolution, Blockers, QueryGeom, StartTM, Dir, DeltaMag, HitBuffer, OutputFlags, ECollisionQueryHitType::Block, bAnyHit);
-	if (bProcessTouches)
-	{
-		ProcessSweeps(Evolution, Touchers, QueryGeom, StartTM, Dir, DeltaMag, HitBuffer, OutputFlags, ECollisionQueryHitType::Touch, bAnyHit);
+		return true;
 	}
 
-	FinalizeQuery(HitBuffer);
-#endif
-}
+	const FTransform StartTM;
+	const FVector StartPoint;
+	const FVector Dir;
+	ChaosInterface::FSQHitBuffer<THitType>& HitBuffer;
+	EHitFlags OutputFlags;
+	const FQueryFilterData& QueryFilterData;
+	ICollisionQueryFilterCallbackBase& QueryCallback;
+	bool bAnyHit;
+	const Chaos::TImplicitObject<float, 3>* QueryGeom;
+};
 
-void ProcessOverlaps(const Chaos::TPBDRigidsEvolutionGBF<float, 3>& Evolution, const TArray<FPreFilterInfo>& PreFilters, const FTransform& WorldTM, const Chaos::TImplicitObject<float,3>& QueryGeom, const float Thickness, FPhysicsHitCallback<FOverlapHit>& HitBuffer, ECollisionQueryHitType HitType, bool bAnyHit)
+void FChaosSQAccelerator::Raycast(const FVector& Start, const FVector& Dir, const float DeltaMagnitude, ChaosInterface::FSQHitBuffer<ChaosInterface::FRaycastHit>& HitBuffer, EHitFlags OutputFlags, const FQueryFilterData& QueryFilterData, ICollisionQueryFilterCallbackBase& QueryCallback) const
 {
-	ensure(bAnyHit || HitType == ECollisionQueryHitType::Touch);	//overlaps only make sense with touch or just a single result
-
 	using namespace Chaos;
-	const TPBDRigidParticles<float, 3>& Particles = Evolution.GetParticles();
-	for (const FPreFilterInfo& PreFilter : PreFilters)
-	{
-		const TRigidTransform<float, 3> GeomWorldTM(Particles.X(PreFilter.ActorIdx), Particles.R(PreFilter.ActorIdx));
-		const TRigidTransform<float, 3> QueryToGeomTM = GeomWorldTM.GetRelativeTransform(WorldTM);
-
-		FOverlapHit Hit;
-		
-		if (OverlapQuery<float,3>(*PreFilter.Geom, GeomWorldTM, QueryGeom, WorldTM, Thickness))
-		{
-			Hit.ActorIdx = PreFilter.ActorIdx;
-			Hit.Shape = PreFilter.Geom;
-			InsertOverlap(HitBuffer, Hit);
+	using namespace ChaosInterface;
+
+	TSQVisitor<TGeometryParticleHandle<float, 3>*, FRaycastHit> RaycastVisitor(Start, Dir, HitBuffer, OutputFlags, QueryFilterData, QueryCallback);
+	const auto& SpatialAcceleration = Evolution.GetCollisionConstraints().GetSpatialAcceleration();
+	HitBuffer.IncFlushCount();
+	SpatialAcceleration.Raycast(Start, Dir, DeltaMagnitude, RaycastVisitor);
+	HitBuffer.DecFlushCount();
+	Evolution.GetCollisionConstraints().ReleaseSpatialAcceleration();
+}
 
-			if (bAnyHit)
-			{
-				return;
-			}
-		}
-	}
+void FChaosSQAccelerator::Sweep(const Chaos::TImplicitObject<float, 3>& QueryGeom, const FTransform& StartTM, const FVector& Dir, const float DeltaMagnitude, ChaosInterface::FSQHitBuffer<ChaosInterface::FSweepHit>& HitBuffer, EHitFlags OutputFlags, const FQueryFilterData& QueryFilterData, ICollisionQueryFilterCallbackBase& QueryCallback) const
+{
+	using namespace Chaos;
+	using namespace ChaosInterface;
+
+	const TBox<float, 3> Bounds = QueryGeom.BoundingBox().TransformedBox(StartTM);
+	const FVector HalfExtents = Bounds.Extents() * 0.5f;
+	TSQVisitor<TGeometryParticleHandle<float, 3>*, FSweepHit> SweepVisitor(StartTM, Dir, HitBuffer, OutputFlags, QueryFilterData, QueryCallback, QueryGeom);
+	const auto& SpatialAcceleration = Evolution.GetCollisionConstraints().GetSpatialAcceleration();
+	HitBuffer.IncFlushCount();
+	SpatialAcceleration.Sweep(Bounds.GetCenter(), Dir, DeltaMagnitude, HalfExtents, SweepVisitor);
+	HitBuffer.DecFlushCount();
+	Evolution.GetCollisionConstraints().ReleaseSpatialAcceleration();
 }
 
-void FChaosSQAccelerator::ChaosOverlap(const Chaos::TImplicitObject<float, 3>& QueryGeom, const FTransform& WorldTM,FPhysicsHitCallback<FOverlapHit>& HitBuffer, const FCollisionFilterData& QueryFilter, const FQueryFilterData& QueryFilterData, ICollisionQueryFilterCallbackBase& QueryCallback) const
+void FChaosSQAccelerator::Overlap(const Chaos::TImplicitObject<float, 3>& QueryGeom, const FTransform& GeomPose, ChaosInterface::FSQHitBuffer<ChaosInterface::FOverlapHit>& HitBuffer, const FQueryFilterData& QueryFilterData, ICollisionQueryFilterCallbackBase& QueryCallback) const
 {
-#if WITH_PHYSX
 	using namespace Chaos;
+	using namespace ChaosInterface;
+
+	const TBox<float, 3> Bounds = QueryGeom.BoundingBox().TransformedBox(GeomPose);
+	TSQVisitor<TGeometryParticleHandle<float, 3>*, FOverlapHit> OverlapVisitor(GeomPose, HitBuffer, QueryFilterData, QueryCallback, QueryGeom);
+	const auto& SpatialAcceleration = Evolution.GetCollisionConstraints().GetSpatialAcceleration();
+	HitBuffer.IncFlushCount();
+	SpatialAcceleration.Overlap(Bounds, OverlapVisitor);
+	HitBuffer.DecFlushCount();
+	Evolution.GetCollisionConstraints().ReleaseSpatialAcceleration();
+}
 
-	TArray<int32> PotentialIntersections;
-	{
-		const ISpatialAcceleration<float, 3>& SpatialAcceleration = Evolution.GetSpatialAcceleration();
-		TBox<float, 3> WorldBounds = QueryGeom.BoundingBox().TransformedBox(WorldTM);
-		PotentialIntersections = SpatialAcceleration.FindAllIntersections(WorldBounds);
-		Evolution.ReleaseSpatialAcceleration();
-	}
+#endif
 
-	TArray<FPreFilterInfo> Overlaps;
-	const bool bAnyHit = QueryFilterData.flags & PxQueryFlag::eANY_HIT;
-	GeneratePrefilterInfo(Evolution, PotentialIntersections, QueryFilter, QueryFilterData, bAnyHit, QueryCallback, Overlaps, Overlaps);
+#if WITH_PHYSX && INCLUDE_CHAOS
+FChaosSQAcceleratorAdapter::FChaosSQAcceleratorAdapter(const Chaos::TPBDRigidsEvolutionGBF<float, 3>& InEvolution)
+	: ChaosSQAccelerator(InEvolution)
+{
+}
 
-	ProcessOverlaps(Evolution, Overlaps, WorldTM, QueryGeom, 0, HitBuffer, bAnyHit ? ECollisionQueryHitType::Block : ECollisionQueryHitType::Touch, bAnyHit);
+void FChaosSQAcceleratorAdapter::Raycast(const FVector& Start, const FVector& Dir, const float DeltaMagnitude, FPhysicsHitCallback<FHitRaycast>& HitBuffer, EHitFlags OutputFlags, const FQueryFilterData& QueryFilterData, ICollisionQueryFilterCallbackBase& QueryCallback) const
+{
+	check(false);
+}
 
-	FinalizeQuery(HitBuffer);
-#endif
+void FChaosSQAcceleratorAdapter::Sweep(const FPhysicsGeometry& QueryGeom, const FTransform& StartTM, const FVector& Dir, const float DeltaMagnitude, FPhysicsHitCallback<FHitSweep>& HitBuffer, EHitFlags OutputFlags, const FQueryFilterData& QueryFilterData, ICollisionQueryFilterCallbackBase& QueryCallback) const
+{
+	check(false);
 }
 
+void FChaosSQAcceleratorAdapter::Overlap(const FPhysicsGeometry& QueryGeom, const FTransform& GeomPose, FPhysicsHitCallback<FHitOverlap>& HitBuffer, const FQueryFilterData& QueryFilterData, ICollisionQueryFilterCallbackBase& QueryCallback) const
+{
+	check(false);
+}
 #endif
 
+
 #if WITH_PHYSX && !WITH_CHAOS
 
 FPhysXSQAccelerator::FPhysXSQAccelerator()
@@ -674,7 +343,7 @@ FPhysXSQAccelerator::FPhysXSQAccelerator(physx::PxScene* InScene)
 
 }
 
-void FPhysXSQAccelerator::Raycast(const FVector& Start, const FVector& Dir, const float DeltaMagnitude, FPhysicsHitCallback<FHitRaycast>& HitBuffer, EHitFlags OutputFlags, FQueryFlags QueryFlags, const FCollisionFilterData& QueryFilter, const FQueryFilterData& QueryFilterData, ICollisionQueryFilterCallbackBase& QueryCallback) const
+void FPhysXSQAccelerator::Raycast(const FVector& Start, const FVector& Dir, const float DeltaMagnitude, FPhysicsHitCallback<FHitRaycast>& HitBuffer, EHitFlags OutputFlags, const FQueryFilterData& QueryFilterData, ICollisionQueryFilterCallbackBase& QueryCallback) const
 {
 	check(Scene);
 
@@ -682,7 +351,7 @@ void FPhysXSQAccelerator::Raycast(const FVector& Start, const FVector& Dir, cons
 	Scene->raycast(Inputs.Start, Inputs.Dir, DeltaMagnitude, HitBuffer, Inputs.OutputFlags, QueryFilterData, &QueryCallback);
 }
 
-void FPhysXSQAccelerator::Sweep(const FPhysicsGeometry& QueryGeom, const FTransform& StartTM, const FVector& Dir, const float DeltaMagnitude, FPhysicsHitCallback<FHitSweep>& HitBuffer, EHitFlags OutputFlags, FQueryFlags QueryFlags, const FCollisionFilterData& QueryFilter, const FQueryFilterData& QueryFilterData, ICollisionQueryFilterCallbackBase& QueryCallback) const
+void FPhysXSQAccelerator::Sweep(const FPhysicsGeometry& QueryGeom, const FTransform& StartTM, const FVector& Dir, const float DeltaMagnitude, FPhysicsHitCallback<FHitSweep>& HitBuffer, EHitFlags OutputFlags, const FQueryFilterData& QueryFilterData, ICollisionQueryFilterCallbackBase& QueryCallback) const
 {
 	check(Scene);
 
@@ -690,7 +359,7 @@ void FPhysXSQAccelerator::Sweep(const FPhysicsGeometry& QueryGeom, const FTransf
 	Scene->sweep(QueryGeom, Inputs.StartTM, Inputs.Dir, DeltaMagnitude, HitBuffer, Inputs.OutputFlags, QueryFilterData, &QueryCallback);
 }
 
-void FPhysXSQAccelerator::Overlap(const FPhysicsGeometry& QueryGeom, const FTransform& GeomPose, FPhysicsHitCallback<FHitOverlap>& HitBuffer, FQueryFlags QueryFlags, const FCollisionFilterData& QueryFilter, const FQueryFilterData& QueryFilterData, ICollisionQueryFilterCallbackBase& QueryCallback) const
+void FPhysXSQAccelerator::Overlap(const FPhysicsGeometry& QueryGeom, const FTransform& GeomPose, FPhysicsHitCallback<FHitOverlap>& HitBuffer, const FQueryFilterData& QueryFilterData, ICollisionQueryFilterCallbackBase& QueryCallback) const
 {
 	check(Scene);
 
diff --git a/Engine/Source/Runtime/PhysicsSQ/Public/SQAccelerator.h b/Engine/Source/Runtime/PhysicsSQ/Public/SQAccelerator.h
--- a/Engine/Source/Runtime/PhysicsSQ/Public/SQAccelerator.h
+++ b/Engine/Source/Runtime/PhysicsSQ/Public/SQAccelerator.h
@@ -6,13 +6,13 @@
 
 #if PHYSICS_INTERFACE_PHYSX
 #include "PhysXInterfaceWrapperCore.h"	//todo remove physx specific include once refactor is done
-#elif PHYSICS_INTERFACE_LLIMMEDIATE
-//#include "Physics/Experimental/LLImmediateInterfaceWrapper.h"
 #elif WITH_CHAOS
 #include "ChaosInterfaceWrapperCore.h"
 #endif
 
+#if INCLUDE_CHAOS
 #include "ChaosSQTypes.h"
+#endif
 
 #if WITH_PHYSX
 namespace physx
@@ -25,74 +25,56 @@ namespace Chaos
 {
 	template <typename, int>
 	class TPBDRigidsEvolutionGBF;
+
+	template <typename, int>
+	class TImplicitObject;
 }
 
 class FSQAccelerator;
 struct FCollisionFilterData;
 struct FCollisionQueryParams;
 struct FCollisionQueryParams;
+class ICollisionQueryFilterCallbackBase;
 
-class FSQAcceleratorEntry
-
+#if INCLUDE_CHAOS
+class PHYSICSSQ_API FChaosSQAccelerator
 {
 public:
-	bool Intersect(const FBoxSphereBounds& Other) const
-	{
-		//return FBoxSphereBounds::BoxesIntersect(Other, Bounds);
-		return true;
-	}
-
-	void* GetPayload() const
-	{
-		return Payload;
-	}
-private:
-	FSQAcceleratorEntry(void* InPayload)
-		: Payload(InPayload) {}
 
-	void* Payload;
-	friend FSQAccelerator;
-};
+	FChaosSQAccelerator(const Chaos::TPBDRigidsEvolutionGBF<float, 3>& InEvolution);
+	virtual ~FChaosSQAccelerator() {};
 
-struct FSQNode
-{
-	TArray<FSQAcceleratorEntry*> Entries;
-};
+	void Raycast(const FVector& Start, const FVector& Dir, const float DeltaMagnitude, ChaosInterface::FSQHitBuffer<ChaosInterface::FRaycastHit>& HitBuffer, EHitFlags OutputFlags, const FQueryFilterData& QueryFilterData, ICollisionQueryFilterCallbackBase& QueryCallback) const;
+	void Sweep(const Chaos::TImplicitObject<float, 3>& QueryGeom, const FTransform& StartTM, const FVector& Dir, const float DeltaMagnitude, ChaosInterface::FSQHitBuffer<ChaosInterface::FSweepHit>& HitBuffer, EHitFlags OutputFlags, const FQueryFilterData& QueryFilterData, ICollisionQueryFilterCallbackBase& QueryCallback) const;
+	void Overlap(const Chaos::TImplicitObject<float, 3>& QueryGeom, const FTransform& GeomPose, ChaosInterface::FSQHitBuffer<ChaosInterface::FOverlapHit>& HitBuffer, const FQueryFilterData& QueryFilterData, ICollisionQueryFilterCallbackBase& QueryCallback) const;
 
-class ICollisionQueryFilterCallbackBase;
+private:
+	const Chaos::TPBDRigidsEvolutionGBF<float, 3>& Evolution;
 
-class PHYSICSSQ_API ISQAccelerator
-{
-public:
-	virtual ~ISQAccelerator() {};
-	virtual void Raycast(const FVector& Start, const FVector& Dir, const float DeltaMagnitude, FPhysicsHitCallback<FHitRaycast>& HitBuffer, EHitFlags OutputFlags, FQueryFlags QueryFlags, const FCollisionFilterData& QueryFilter, const FQueryFilterData& QueryFilterData, ICollisionQueryFilterCallbackBase& QueryCallback) const = 0;
-	virtual void Sweep(const FPhysicsGeometry& QueryGeom, const FTransform& StartTM, const FVector& Dir, const float DeltaMagnitude, FPhysicsHitCallback<FHitSweep>& HitBuffer, EHitFlags OutputFlags, FQueryFlags QueryFlags, const FCollisionFilterData& QueryFilter, const FQueryFilterData& Params, ICollisionQueryFilterCallbackBase& QueryCallback) const = 0;
-	virtual void Overlap(const FPhysicsGeometry& QueryGeom, const FTransform& GeomPose, FPhysicsHitCallback<FHitOverlap>& HitBuffer, FQueryFlags QueryFlags, const FCollisionFilterData& QueryFilter, const FQueryFilterData& QueryFilterData, ICollisionQueryFilterCallbackBase& QueryCallback) const = 0;
 };
+#endif
 
-class PHYSICSSQ_API FSQAccelerator : public ISQAccelerator
+// An interface to the scene query accelerator that allows us to run queries against either PhysX or Chaos
+// when compiling WITH_PHYSX and INCLUDE_CHAOS.
+// This was used in the 2019 GDC demos and is now broken. To make it work again, we would need to implement
+// the FChaosSQAcceleratorAdapter below to use its internal SQ accelerator and convert the inputs and outputs
+// from/to PhysX types.
+class PHYSICSSQ_API ISQAccelerator
 {
 public:
-	FSQAcceleratorEntry* AddEntry(void* Payload);
-	void RemoveEntry(FSQAcceleratorEntry* Entry);
-	void GetNodes(TArray<const FSQNode*>& NodesFound) const;
-
-	virtual void Raycast(const FVector& Start, const FVector& Dir, const float DeltaMagnitude, FPhysicsHitCallback<FHitRaycast>& HitBuffer, EHitFlags OutputFlags, FQueryFlags QueryFlags, const FCollisionFilterData& QueryFilter, const FQueryFilterData& QueryFilterData, ICollisionQueryFilterCallbackBase& QueryCallback) const override;
-	virtual void Sweep(const FPhysicsGeometry& QueryGeom, const FTransform& StartTM, const FVector& Dir, const float DeltaMagnitude, FPhysicsHitCallback<FHitSweep>& HitBuffer, EHitFlags OutputFlags, FQueryFlags QueryFlags, const FCollisionFilterData& QueryFilter, const FQueryFilterData& QueryFilterData, ICollisionQueryFilterCallbackBase& QueryCallback) const override;
-	virtual void Overlap(const FPhysicsGeometry& QueryGeom, const FTransform& GeomPose, FPhysicsHitCallback<FHitOverlap>& HitBuffer, FQueryFlags QueryFlags, const FCollisionFilterData& QueryFilter, const FQueryFilterData& QueryFilterData, ICollisionQueryFilterCallbackBase& QueryCallback) const override;
-
-	virtual ~FSQAccelerator() override;
-private:
-	TArray<FSQNode*> Nodes;
+	virtual ~ISQAccelerator() {};
+	virtual void Raycast(const FVector& Start, const FVector& Dir, const float DeltaMagnitude, FPhysicsHitCallback<FHitRaycast>& HitBuffer, EHitFlags OutputFlags, const FQueryFilterData& QueryFilterData, ICollisionQueryFilterCallbackBase& QueryCallback) const = 0;
+	virtual void Sweep(const FPhysicsGeometry& QueryGeom, const FTransform& StartTM, const FVector& Dir, const float DeltaMagnitude, FPhysicsHitCallback<FHitSweep>& HitBuffer, EHitFlags OutputFlags, const FQueryFilterData& QueryFilterData, ICollisionQueryFilterCallbackBase& QueryCallback) const = 0;
+	virtual void Overlap(const FPhysicsGeometry& QueryGeom, const FTransform& GeomPose, FPhysicsHitCallback<FHitOverlap>& HitBuffer, const FQueryFilterData& QueryFilterData, ICollisionQueryFilterCallbackBase& QueryCallback) const = 0;
 };
 
 class PHYSICSSQ_API FSQAcceleratorUnion : public ISQAccelerator
 {
 public:
 
-	virtual void Raycast(const FVector& Start, const FVector& Dir, const float DeltaMagnitude, FPhysicsHitCallback<FHitRaycast>& HitBuffer, EHitFlags OutputFlags, FQueryFlags QueryFlags, const FCollisionFilterData& QueryFilter, const FQueryFilterData& QueryFilterData, ICollisionQueryFilterCallbackBase& QueryCallback) const override;
-	virtual void Sweep(const FPhysicsGeometry& QueryGeom, const FTransform& StartTM, const FVector& Dir, const float DeltaMagnitude, FPhysicsHitCallback<FHitSweep>& HitBuffer, EHitFlags OutputFlags, FQueryFlags QueryFlags, const FCollisionFilterData& QueryFilter, const FQueryFilterData& QueryFilterData, ICollisionQueryFilterCallbackBase& QueryCallback) const override;
-	virtual void Overlap(const FPhysicsGeometry& QueryGeom, const FTransform& GeomPose, FPhysicsHitCallback<FHitOverlap>& HitBuffer, FQueryFlags QueryFlags, const FCollisionFilterData& QueryFilter, const FQueryFilterData& QueryFilterData, ICollisionQueryFilterCallbackBase& QueryCallback) const override;
+	virtual void Raycast(const FVector& Start, const FVector& Dir, const float DeltaMagnitude, FPhysicsHitCallback<FHitRaycast>& HitBuffer, EHitFlags OutputFlags, const FQueryFilterData& QueryFilterData, ICollisionQueryFilterCallbackBase& QueryCallback) const override;
+	virtual void Sweep(const FPhysicsGeometry& QueryGeom, const FTransform& StartTM, const FVector& Dir, const float DeltaMagnitude, FPhysicsHitCallback<FHitSweep>& HitBuffer, EHitFlags OutputFlags, const FQueryFilterData& QueryFilterData, ICollisionQueryFilterCallbackBase& QueryCallback) const override;
+	virtual void Overlap(const FPhysicsGeometry& QueryGeom, const FTransform& GeomPose, FPhysicsHitCallback<FHitOverlap>& HitBuffer, const FQueryFilterData& QueryFilterData, ICollisionQueryFilterCallbackBase& QueryCallback) const override;
 
 	void AddSQAccelerator(ISQAccelerator* InAccelerator);
 	void RemoveSQAccelerator(ISQAccelerator* AcceleratorToRemove);
@@ -101,40 +83,26 @@ private:
 	TArray<ISQAccelerator*> Accelerators;
 };
 
-#if INCLUDE_CHAOS
-class PHYSICSSQ_API FChaosSQAccelerator : public ISQAccelerator
+#if WITH_PHYSX && INCLUDE_CHAOS
+// A Chaos Query Accelerator with a PhysX API
+// TODO: Not implemented - required to make GDC 2019 demos work again.
+class PHYSICSSQ_API FChaosSQAcceleratorAdapter : public ISQAccelerator
 {
 public:
 
-	FChaosSQAccelerator(const Chaos::TPBDRigidsEvolutionGBF<float, 3>& InEvolution);
-	virtual ~FChaosSQAccelerator() {};
-	virtual void ChaosSweep(const Chaos::TImplicitObject<float, 3>& QueryGeom, const FTransform& StartTM, const FVector& Dir, const float DeltaMagnitude, FPhysicsHitCallback<FSweepHit>& HitBuffer, EHitFlags OutputFlags, const FCollisionFilterData& QueryFilter, const FQueryFilterData& QueryFilterData, ICollisionQueryFilterCallbackBase& QueryCallback) const;
-	virtual void ChaosRaycast(const FVector& StartPoint, const FVector& Dir, const float DeltaMagnitude, FPhysicsHitCallback<FRaycastHit>& HitBuffer, EHitFlags OutputFlags, const FCollisionFilterData& QueryFilter, const FQueryFilterData& QueryFilterData, ICollisionQueryFilterCallbackBase& QueryCallback) const;
-	virtual void ChaosOverlap(const Chaos::TImplicitObject<float, 3>& QueryGeom, const FTransform& WorldTM,FPhysicsHitCallback<FOverlapHit>& HitBuffer, const FCollisionFilterData& QueryFilter, const FQueryFilterData& QueryFilterData, ICollisionQueryFilterCallbackBase& QueryCallback) const;
-
-	virtual void Raycast(const FVector& Start, const FVector& Dir,const float DeltaMagnitude, FPhysicsHitCallback<FHitRaycast>& HitBuffer, EHitFlags OutputFlags, FQueryFlags QueryFlags, const FCollisionFilterData& QueryFilter, const FQueryFilterData& QueryFilterData, ICollisionQueryFilterCallbackBase& QueryCallback) const override
-	{
-		check(false);
-	}
-
-	virtual void Sweep(const FPhysicsGeometry& QueryGeom, const FTransform& StartTM, const FVector& Dir, const float DeltaMagnitude, FPhysicsHitCallback<FHitSweep>& HitBuffer, EHitFlags OutputFlags, FQueryFlags QueryFlags, const FCollisionFilterData& QueryFilter, const FQueryFilterData& QueryFilterData, ICollisionQueryFilterCallbackBase& QueryCallback) const override
-	{
-		check(false);
-	}
+	FChaosSQAcceleratorAdapter(const Chaos::TPBDRigidsEvolutionGBF<float, 3>& InEvolution);
+	virtual ~FChaosSQAcceleratorAdapter() {};
 
-	virtual void Overlap(const FPhysicsGeometry& QueryGeom, const FTransform& GeomPose, FPhysicsHitCallback<FHitOverlap>& HitBuffer, FQueryFlags QueryFlags, const FCollisionFilterData& QueryFilter, const FQueryFilterData& QueryFilterData, ICollisionQueryFilterCallbackBase& QueryCallback) const override
-	{
-		check(false);
-	}
-
-	void SetScene(FPhysScene* InScene);
+	virtual void Raycast(const FVector& Start, const FVector& Dir, const float DeltaMagnitude, FPhysicsHitCallback<FHitRaycast>& HitBuffer, EHitFlags OutputFlags, const FQueryFilterData& QueryFilterData, ICollisionQueryFilterCallbackBase& QueryCallback) const override;
+	virtual void Sweep(const FPhysicsGeometry& QueryGeom, const FTransform& StartTM, const FVector& Dir, const float DeltaMagnitude, FPhysicsHitCallback<FHitSweep>& HitBuffer, EHitFlags OutputFlags, const FQueryFilterData& QueryFilterData, ICollisionQueryFilterCallbackBase& QueryCallback) const override;
+	virtual void Overlap(const FPhysicsGeometry& QueryGeom, const FTransform& GeomPose, FPhysicsHitCallback<FHitOverlap>& HitBuffer, const FQueryFilterData& QueryFilterData, ICollisionQueryFilterCallbackBase& QueryCallback) const override;
 
 private:
-	const Chaos::TPBDRigidsEvolutionGBF<float, 3>& Evolution;
-
+	FChaosSQAccelerator ChaosSQAccelerator;
 };
 #endif
 
+
 #if WITH_PHYSX && !WITH_CHAOS
 class PHYSICSSQ_API FPhysXSQAccelerator : public ISQAccelerator
 {
@@ -144,9 +112,9 @@ public:
 	FPhysXSQAccelerator(physx::PxScene* InScene);
 	virtual ~FPhysXSQAccelerator() {};
 
-	virtual void Raycast(const FVector& Start, const FVector& Dir, const float DeltaMagnitude, FPhysicsHitCallback<FHitRaycast>& HitBuffer, EHitFlags OutputFlags, FQueryFlags QueryFlags, const FCollisionFilterData& QueryFilter, const FQueryFilterData& QueryFilterData, ICollisionQueryFilterCallbackBase& QueryCallback) const override;
-	virtual void Sweep(const FPhysicsGeometry& QueryGeom, const FTransform& StartTM, const FVector& Dir, const float DeltaMagnitude, FPhysicsHitCallback<FHitSweep>& HitBuffer, EHitFlags OutputFlags, FQueryFlags QueryFlags, const FCollisionFilterData& QueryFilter, const FQueryFilterData& QueryFilterData, ICollisionQueryFilterCallbackBase& QueryCallback) const override;
-	virtual void Overlap(const FPhysicsGeometry& QueryGeom, const FTransform& GeomPose, FPhysicsHitCallback<FHitOverlap>& HitBuffer, FQueryFlags QueryFlags, const FCollisionFilterData& QueryFilter, const FQueryFilterData& QueryFilterData, ICollisionQueryFilterCallbackBase& QueryCallback) const override;
+	virtual void Raycast(const FVector& Start, const FVector& Dir, const float DeltaMagnitude, FPhysicsHitCallback<FHitRaycast>& HitBuffer, EHitFlags OutputFlags, const FQueryFilterData& QueryFilterData, ICollisionQueryFilterCallbackBase& QueryCallback) const override;
+	virtual void Sweep(const FPhysicsGeometry& QueryGeom, const FTransform& StartTM, const FVector& Dir, const float DeltaMagnitude, FPhysicsHitCallback<FHitSweep>& HitBuffer, EHitFlags OutputFlags, const FQueryFilterData& QueryFilterData, ICollisionQueryFilterCallbackBase& QueryCallback) const override;
+	virtual void Overlap(const FPhysicsGeometry& QueryGeom, const FTransform& GeomPose, FPhysicsHitCallback<FHitOverlap>& HitBuffer, const FQueryFilterData& QueryFilterData, ICollisionQueryFilterCallbackBase& QueryCallback) const override;
 
 	void SetScene(physx::PxScene* InScene);
 
diff --git a/Engine/Source/Runtime/PhysicsSQ/Public/SQVerifier.h b/Engine/Source/Runtime/PhysicsSQ/Public/SQVerifier.h
--- a/Engine/Source/Runtime/PhysicsSQ/Public/SQVerifier.h
+++ b/Engine/Source/Runtime/PhysicsSQ/Public/SQVerifier.h
@@ -23,7 +23,7 @@ bool SQComparisonHelper(FPhysTestSerializer& Serializer, bool bEnsureOnMismatch
 	{
 	case FSQCapture::ESQType::Raycast:
 	{
-		TUniquePtr<FDynamicHitBuffer<PxRaycastHit>> PxHitBuffer = MakeUnique<FDynamicHitBuffer<PxRaycastHit>>();
+		auto PxHitBuffer = MakeUnique<PhysXInterface::FDynamicHitBuffer<PxRaycastHit>>();
 		Serializer.GetPhysXData()->raycast(U2PVector(CapturedSQ.StartPoint), U2PVector(CapturedSQ.Dir), CapturedSQ.DeltaMag, *PxHitBuffer, U2PHitFlags(CapturedSQ.OutputFlags.HitFlags), CapturedSQ.QueryFilterData, CapturedSQ.FilterCallback.Get());
 
 		bTestPassed &= SQ_REPLAY_TEST(PxHitBuffer->hasBlock == CapturedSQ.PhysXRaycastBuffer.hasBlock);
@@ -42,11 +42,11 @@ bool SQComparisonHelper(FPhysTestSerializer& Serializer, bool bEnsureOnMismatch
 			bTestPassed &= SQ_REPLAY_TEST(FMath::IsNearlyEqual(PxHitBuffer->block.position.z, CapturedSQ.PhysXRaycastBuffer.block.position.z, DistanceTolerance));
 		}
 
-		TUniquePtr<FDynamicHitBuffer<FRaycastHit>> ChaosHitBuffer = MakeUnique<FDynamicHitBuffer<FRaycastHit>>();
+		auto ChaosHitBuffer = MakeUnique<ChaosInterface::FSQHitBuffer<ChaosInterface::FRaycastHit>>();
 		FChaosSQAccelerator SQAccelerator(*Serializer.GetChaosData());
-		SQAccelerator.ChaosRaycast(CapturedSQ.StartPoint, CapturedSQ.Dir, CapturedSQ.DeltaMag, *ChaosHitBuffer, CapturedSQ.OutputFlags.HitFlags, FCollisionFilterData(), CapturedSQ.QueryFilterData, *CapturedSQ.FilterCallback);
+		SQAccelerator.Raycast(CapturedSQ.StartPoint, CapturedSQ.Dir, CapturedSQ.DeltaMag, *ChaosHitBuffer, CapturedSQ.OutputFlags.HitFlags, CapturedSQ.QueryFilterData, *CapturedSQ.FilterCallback);
 
-		bTestPassed &= SQ_REPLAY_TEST(ChaosHitBuffer->hasBlock == CapturedSQ.PhysXRaycastBuffer.hasBlock);
+		bTestPassed &= SQ_REPLAY_TEST(ChaosHitBuffer->HasBlockingHit() == CapturedSQ.PhysXRaycastBuffer.hasBlock);
 		bTestPassed &= SQ_REPLAY_TEST(ChaosHitBuffer->GetNumHits() == CapturedSQ.PhysXRaycastBuffer.GetNumHits());
 		for (int32 Idx = 0; Idx < ChaosHitBuffer->GetNumHits(); ++Idx)
 		{
@@ -56,22 +56,22 @@ bool SQComparisonHelper(FPhysTestSerializer& Serializer, bool bEnsureOnMismatch
 			//bTestPassed &= SQ_REPLAY_TEST(FMath::IsNearlyEqual(ChaosHitBuffer.GetHits()[Idx].WorldNormal.Z, CapturedSQ.PhysXRaycastBuffer.GetHits()[Idx].normal.z));
 		}
 
-		if (ChaosHitBuffer->hasBlock)
+		if (ChaosHitBuffer->HasBlockingHit())
 		{
-			bTestPassed &= SQ_REPLAY_TEST(FMath::IsNearlyEqual(ChaosHitBuffer->block.WorldPosition.X, CapturedSQ.PhysXRaycastBuffer.block.position.x, DistanceTolerance));
-			bTestPassed &= SQ_REPLAY_TEST(FMath::IsNearlyEqual(ChaosHitBuffer->block.WorldPosition.Y, CapturedSQ.PhysXRaycastBuffer.block.position.y, DistanceTolerance));
-			bTestPassed &= SQ_REPLAY_TEST(FMath::IsNearlyEqual(ChaosHitBuffer->block.WorldPosition.Z, CapturedSQ.PhysXRaycastBuffer.block.position.z, DistanceTolerance));
+			bTestPassed &= SQ_REPLAY_TEST(FMath::IsNearlyEqual(ChaosHitBuffer->GetBlock()->WorldPosition.X, CapturedSQ.PhysXRaycastBuffer.block.position.x, DistanceTolerance));
+			bTestPassed &= SQ_REPLAY_TEST(FMath::IsNearlyEqual(ChaosHitBuffer->GetBlock()->WorldPosition.Y, CapturedSQ.PhysXRaycastBuffer.block.position.y, DistanceTolerance));
+			bTestPassed &= SQ_REPLAY_TEST(FMath::IsNearlyEqual(ChaosHitBuffer->GetBlock()->WorldPosition.Z, CapturedSQ.PhysXRaycastBuffer.block.position.z, DistanceTolerance));
 
-			bTestPassed &= SQ_REPLAY_TEST(FMath::IsNearlyEqual(ChaosHitBuffer->block.WorldNormal.X, CapturedSQ.PhysXRaycastBuffer.block.normal.x, NormalTolerance));
-			bTestPassed &= SQ_REPLAY_TEST(FMath::IsNearlyEqual(ChaosHitBuffer->block.WorldNormal.Y, CapturedSQ.PhysXRaycastBuffer.block.normal.y, NormalTolerance));
-			bTestPassed &= SQ_REPLAY_TEST(FMath::IsNearlyEqual(ChaosHitBuffer->block.WorldNormal.Z, CapturedSQ.PhysXRaycastBuffer.block.normal.z, NormalTolerance));
+			bTestPassed &= SQ_REPLAY_TEST(FMath::IsNearlyEqual(ChaosHitBuffer->GetBlock()->WorldNormal.X, CapturedSQ.PhysXRaycastBuffer.block.normal.x, NormalTolerance));
+			bTestPassed &= SQ_REPLAY_TEST(FMath::IsNearlyEqual(ChaosHitBuffer->GetBlock()->WorldNormal.Y, CapturedSQ.PhysXRaycastBuffer.block.normal.y, NormalTolerance));
+			bTestPassed &= SQ_REPLAY_TEST(FMath::IsNearlyEqual(ChaosHitBuffer->GetBlock()->WorldNormal.Z, CapturedSQ.PhysXRaycastBuffer.block.normal.z, NormalTolerance));
 		}
 		break;
 	}
 	case FSQCapture::ESQType::Sweep:
 	{
 		//For sweep there are many solutions (many contacts possible) so we only bother testing for Distance
-		TUniquePtr<FDynamicHitBuffer<PxSweepHit>> PxHitBuffer = MakeUnique<FDynamicHitBuffer<PxSweepHit>>();
+		auto PxHitBuffer = MakeUnique<PhysXInterface::FDynamicHitBuffer<PxSweepHit>>();
 		Serializer.GetPhysXData()->sweep(CapturedSQ.PhysXGeometry.any(), U2PTransform(CapturedSQ.StartTM), U2PVector(CapturedSQ.Dir), CapturedSQ.DeltaMag, *PxHitBuffer, U2PHitFlags(CapturedSQ.OutputFlags.HitFlags), CapturedSQ.QueryFilterData, CapturedSQ.FilterCallback.Get());
 
 		bTestPassed &= SQ_REPLAY_TEST(PxHitBuffer->hasBlock == CapturedSQ.PhysXSweepBuffer.hasBlock);
@@ -83,11 +83,11 @@ bool SQComparisonHelper(FPhysTestSerializer& Serializer, bool bEnsureOnMismatch
 			//bTestPassed &= SQ_REPLAY_TEST(FMath::IsNearlyEqual(PxHitBuffer.GetHits()[Idx].normal.z, CapturedSQ.PhysXSweepBuffer.GetHits()[Idx].normal.z, DistanceTolerance));
 		}
 
-		TUniquePtr<FDynamicHitBuffer<FSweepHit>> ChaosHitBuffer = MakeUnique<FDynamicHitBuffer<FSweepHit>>();
+		auto ChaosHitBuffer = MakeUnique<ChaosInterface::FSQHitBuffer<ChaosInterface::FSweepHit>>();
 		FChaosSQAccelerator SQAccelerator(*Serializer.GetChaosData());
-		SQAccelerator.ChaosSweep(*CapturedSQ.ChaosGeometry, CapturedSQ.StartTM, CapturedSQ.Dir, CapturedSQ.DeltaMag, *ChaosHitBuffer, CapturedSQ.OutputFlags.HitFlags, FCollisionFilterData(), CapturedSQ.QueryFilterData, *CapturedSQ.FilterCallback);
+		SQAccelerator.Sweep(*CapturedSQ.ChaosGeometry, CapturedSQ.StartTM, CapturedSQ.Dir, CapturedSQ.DeltaMag, *ChaosHitBuffer, CapturedSQ.OutputFlags.HitFlags, CapturedSQ.QueryFilterData, *CapturedSQ.FilterCallback);
 
-		bTestPassed &= SQ_REPLAY_TEST(ChaosHitBuffer->hasBlock == CapturedSQ.PhysXSweepBuffer.hasBlock);
+		bTestPassed &= SQ_REPLAY_TEST(ChaosHitBuffer->HasBlockingHit() == CapturedSQ.PhysXSweepBuffer.hasBlock);
 		bTestPassed &= SQ_REPLAY_TEST(ChaosHitBuffer->GetNumHits() == CapturedSQ.PhysXSweepBuffer.GetNumHits());
 		for (int32 Idx = 0; Idx < ChaosHitBuffer->GetNumHits(); ++Idx)
 		{
@@ -97,22 +97,22 @@ bool SQComparisonHelper(FPhysTestSerializer& Serializer, bool bEnsureOnMismatch
 			//bTestPassed &= SQ_REPLAY_TEST(FMath::IsNearlyEqual(ChaosHitBuffer.GetHits()[Idx].WorldNormal.Z, CapturedSQ.PhysXSweepBuffer.GetHits()[Idx].normal.z));
 		}
 
-		if (ChaosHitBuffer->hasBlock)
+		if (ChaosHitBuffer->HasBlockingHit())
 		{
-			bTestPassed &= SQ_REPLAY_TEST(FMath::IsNearlyEqual(ChaosHitBuffer->block.Distance, CapturedSQ.PhysXSweepBuffer.block.distance, DistanceTolerance));
+			bTestPassed &= SQ_REPLAY_TEST(FMath::IsNearlyEqual(ChaosHitBuffer->GetBlock()->Distance, CapturedSQ.PhysXSweepBuffer.block.distance, DistanceTolerance));
 		}
 		break;
 	}
 	case FSQCapture::ESQType::Overlap:
 	{
-		TUniquePtr<FDynamicHitBuffer<PxOverlapHit>> PxHitBuffer = MakeUnique<FDynamicHitBuffer<PxOverlapHit>>();
+		auto PxHitBuffer = MakeUnique<PhysXInterface::FDynamicHitBuffer<PxOverlapHit>>();
 		Serializer.GetPhysXData()->overlap(CapturedSQ.PhysXGeometry.any(), U2PTransform(CapturedSQ.StartTM), *PxHitBuffer, CapturedSQ.QueryFilterData, CapturedSQ.FilterCallback.Get());
 
 		bTestPassed &= SQ_REPLAY_TEST(PxHitBuffer->GetNumHits() == CapturedSQ.PhysXOverlapBuffer.GetNumHits());
 
-		TUniquePtr<FDynamicHitBuffer<FOverlapHit>> ChaosHitBuffer = MakeUnique<FDynamicHitBuffer<FOverlapHit>>();
+		auto ChaosHitBuffer = MakeUnique<ChaosInterface::FSQHitBuffer<ChaosInterface::FOverlapHit>>();
 		FChaosSQAccelerator SQAccelerator(*Serializer.GetChaosData());
-		SQAccelerator.ChaosOverlap(*CapturedSQ.ChaosGeometry, CapturedSQ.StartTM, *ChaosHitBuffer, FCollisionFilterData(), CapturedSQ.QueryFilterData, *CapturedSQ.FilterCallback);
+		SQAccelerator.Overlap(*CapturedSQ.ChaosGeometry, CapturedSQ.StartTM, *ChaosHitBuffer, CapturedSQ.QueryFilterData, *CapturedSQ.FilterCallback);
 
 		bTestPassed &= SQ_REPLAY_TEST(ChaosHitBuffer->GetNumHits() == CapturedSQ.PhysXOverlapBuffer.GetNumHits());
 		break;
@@ -136,7 +136,7 @@ void SQPerfComparisonHelper(const FString& TestName, FPhysTestSerializer& Serial
 	{
 		for (double i = 0; i < NumIterations; ++i)
 		{
-			TUniquePtr<FDynamicHitBuffer<PxRaycastHit>> PxHitBuffer = MakeUnique<FDynamicHitBuffer<PxRaycastHit>>();
+			auto PxHitBuffer = MakeUnique<PhysXInterface::FDynamicHitBuffer<PxRaycastHit>>();
 			FDurationTimer Timer(PhysXSum);
 			Timer.Start();
 			Serializer.GetPhysXData()->raycast(U2PVector(CapturedSQ.StartPoint), U2PVector(CapturedSQ.Dir), CapturedSQ.DeltaMag, *PxHitBuffer, U2PHitFlags(CapturedSQ.OutputFlags.HitFlags), CapturedSQ.QueryFilterData, CapturedSQ.FilterCallback.Get());
@@ -146,10 +146,10 @@ void SQPerfComparisonHelper(const FString& TestName, FPhysTestSerializer& Serial
 		FChaosSQAccelerator SQAccelerator(*Serializer.GetChaosData());
 		for (double i = 0; i < NumIterations; ++i)
 		{
-			TUniquePtr<FDynamicHitBuffer<FRaycastHit>> ChaosHitBuffer = MakeUnique<FDynamicHitBuffer<FRaycastHit>>();
+			auto ChaosHitBuffer = MakeUnique<ChaosInterface::FSQHitBuffer<ChaosInterface::FRaycastHit>>();
 			FDurationTimer Timer(ChaosSum);
 			Timer.Start();
-			SQAccelerator.ChaosRaycast(CapturedSQ.StartPoint, CapturedSQ.Dir, CapturedSQ.DeltaMag, *ChaosHitBuffer, CapturedSQ.OutputFlags.HitFlags, FCollisionFilterData(), CapturedSQ.QueryFilterData, *CapturedSQ.FilterCallback);
+			SQAccelerator.Raycast(CapturedSQ.StartPoint, CapturedSQ.Dir, CapturedSQ.DeltaMag, *ChaosHitBuffer, CapturedSQ.OutputFlags.HitFlags, CapturedSQ.QueryFilterData, *CapturedSQ.FilterCallback);
 			Timer.Stop();
 		}
 
@@ -159,7 +159,7 @@ void SQPerfComparisonHelper(const FString& TestName, FPhysTestSerializer& Serial
 	{
 		for (double i = 0; i < NumIterations; ++i)
 		{
-			TUniquePtr<FDynamicHitBuffer<PxSweepHit>> PxHitBuffer = MakeUnique<FDynamicHitBuffer<PxSweepHit>>();
+			auto PxHitBuffer = MakeUnique<PhysXInterface::FDynamicHitBuffer<PxSweepHit>>();
 			FDurationTimer Timer(PhysXSum);
 			Timer.Start();
 			Serializer.GetPhysXData()->sweep(CapturedSQ.PhysXGeometry.any(), U2PTransform(CapturedSQ.StartTM), U2PVector(CapturedSQ.Dir), CapturedSQ.DeltaMag, *PxHitBuffer, U2PHitFlags(CapturedSQ.OutputFlags.HitFlags), CapturedSQ.QueryFilterData, CapturedSQ.FilterCallback.Get());
@@ -169,10 +169,10 @@ void SQPerfComparisonHelper(const FString& TestName, FPhysTestSerializer& Serial
 		FChaosSQAccelerator SQAccelerator(*Serializer.GetChaosData());
 		for (double i = 0; i < NumIterations; ++i)
 		{
-			TUniquePtr<FDynamicHitBuffer<FSweepHit>> ChaosHitBuffer = MakeUnique<FDynamicHitBuffer<FSweepHit>>();
+			auto ChaosHitBuffer = MakeUnique<ChaosInterface::FSQHitBuffer<ChaosInterface::FSweepHit>>();
 			FDurationTimer Timer(ChaosSum);
 			Timer.Start();
-			SQAccelerator.ChaosSweep(*CapturedSQ.ChaosGeometry, CapturedSQ.StartTM, CapturedSQ.Dir, CapturedSQ.DeltaMag, *ChaosHitBuffer, CapturedSQ.OutputFlags.HitFlags, FCollisionFilterData(), CapturedSQ.QueryFilterData, *CapturedSQ.FilterCallback);
+			SQAccelerator.Sweep(*CapturedSQ.ChaosGeometry, CapturedSQ.StartTM, CapturedSQ.Dir, CapturedSQ.DeltaMag, *ChaosHitBuffer, CapturedSQ.OutputFlags.HitFlags, CapturedSQ.QueryFilterData, *CapturedSQ.FilterCallback);
 			Timer.Stop();
 		}
 
@@ -182,7 +182,7 @@ void SQPerfComparisonHelper(const FString& TestName, FPhysTestSerializer& Serial
 	{
 		for (double i = 0; i < NumIterations; ++i)
 		{
-			TUniquePtr<FDynamicHitBuffer<PxOverlapHit>> PxHitBuffer = MakeUnique<FDynamicHitBuffer<PxOverlapHit>>();
+			auto PxHitBuffer = MakeUnique<PhysXInterface::FDynamicHitBuffer<PxOverlapHit>>();
 			FDurationTimer Timer(PhysXSum);
 			Timer.Start();
 			Serializer.GetPhysXData()->overlap(CapturedSQ.PhysXGeometry.any(), U2PTransform(CapturedSQ.StartTM), *PxHitBuffer, CapturedSQ.QueryFilterData, CapturedSQ.FilterCallback.Get());
@@ -192,10 +192,10 @@ void SQPerfComparisonHelper(const FString& TestName, FPhysTestSerializer& Serial
 		FChaosSQAccelerator SQAccelerator(*Serializer.GetChaosData());
 		for (double i = 0; i < NumIterations; ++i)
 		{
-			TUniquePtr<FDynamicHitBuffer<FOverlapHit>> ChaosHitBuffer = MakeUnique<FDynamicHitBuffer<FOverlapHit>>();
+			auto ChaosHitBuffer = MakeUnique<ChaosInterface::FSQHitBuffer<ChaosInterface::FOverlapHit>>();
 			FDurationTimer Timer(ChaosSum);
 			Timer.Start();
-			SQAccelerator.ChaosOverlap(*CapturedSQ.ChaosGeometry, CapturedSQ.StartTM, *ChaosHitBuffer, FCollisionFilterData(), CapturedSQ.QueryFilterData, *CapturedSQ.FilterCallback);
+			SQAccelerator.Overlap(*CapturedSQ.ChaosGeometry, CapturedSQ.StartTM, *ChaosHitBuffer, CapturedSQ.QueryFilterData, *CapturedSQ.FilterCallback);
 			Timer.Stop();
 		}
 		break;
diff --git a/Engine/Source/Runtime/RHI/Public/RHIUtilities.h b/Engine/Source/Runtime/RHI/Public/RHIUtilities.h
--- a/Engine/Source/Runtime/RHI/Public/RHIUtilities.h
+++ b/Engine/Source/Runtime/RHI/Public/RHIUtilities.h
@@ -55,6 +55,75 @@ static inline bool IsStencilFormat(EPixelFormat Format)
 	return false;
 }
 
+/** Encapsulates a GPU read/write texture with its UAV and SRV. */
+struct FTextureRWBuffer
+{
+	FTexture2DRHIRef Buffer;
+	FUnorderedAccessViewRHIRef UAV;
+	FShaderResourceViewRHIRef SRV;
+	uint32 NumBytes;
+
+	FTextureRWBuffer()
+		: NumBytes(0)
+	{}
+
+	~FTextureRWBuffer()
+	{
+		Release();
+	}
+
+	// @param AdditionalUsage passed down to RHICreateVertexBuffer(), get combined with "BUF_UnorderedAccess | BUF_ShaderResource" e.g. BUF_Static
+	void Initialize(uint32 SizeX, uint32 SizeY)
+	{
+		check(GMaxRHIFeatureLevel == ERHIFeatureLevel::SM5
+			|| IsVulkanPlatform(GMaxRHIShaderPlatform)
+			|| IsMetalPlatform(GMaxRHIShaderPlatform)
+			|| (GMaxRHIFeatureLevel == ERHIFeatureLevel::ES3_1 && GSupportsResourceView)
+		);		
+				
+		// #todo(dmp): hardcoded for now...
+		uint32 BlockBytes = 16;
+
+		NumBytes = SizeX * SizeY * BlockBytes;
+
+		FRHIResourceCreateInfo CreateInfo;
+		Buffer = RHICreateTexture2D(
+			SizeX, SizeY, PF_A32B32G32R32F,
+			/*NumMips=*/ 1,
+			1,
+			/*Flags=*/ TexCreate_ShaderResource | TexCreate_UAV,
+			/*BulkData=*/ CreateInfo);
+
+							
+		UAV = RHICreateUnorderedAccessView(Buffer, 0);
+		SRV = RHICreateShaderResourceView(Buffer, 0);
+	}
+
+	void AcquireTransientResource()
+	{
+		RHIAcquireTransientResource(Buffer);
+	}
+	void DiscardTransientResource()
+	{
+		RHIDiscardTransientResource(Buffer);
+	}
+
+	void Release()
+	{
+		int32 BufferRefCount = Buffer ? Buffer->GetRefCount() : -1;
+
+		if (BufferRefCount == 1)
+		{
+			DiscardTransientResource();
+		}
+
+		NumBytes = 0;
+		Buffer.SafeRelease();
+		UAV.SafeRelease();
+		SRV.SafeRelease();
+	}
+};
+
 
 /** Encapsulates a GPU read/write buffer with its UAV and SRV. */
 struct FRWBuffer
diff --git a/Engine/Source/Runtime/UtilityShaders/Public/ClearQuad.h b/Engine/Source/Runtime/UtilityShaders/Public/ClearQuad.h
--- a/Engine/Source/Runtime/UtilityShaders/Public/ClearQuad.h
+++ b/Engine/Source/Runtime/UtilityShaders/Public/ClearQuad.h
@@ -44,6 +44,7 @@ struct FClearQuadCallbacks
 };
 
 extern UTILITYSHADERS_API void ClearUAV(FRHICommandList& RHICmdList, const FRWBufferStructured& StructuredBuffer, uint32 Value);
+extern UTILITYSHADERS_API void ClearUAV(FRHICommandList& RHICmdList, const FTextureRWBuffer& Buffer, FLinearColor Value);
 extern UTILITYSHADERS_API void ClearUAV(FRHICommandList& RHICmdList, const FRWBuffer& Buffer, uint32 Value, bool bBarriers = true);
 extern UTILITYSHADERS_API void ClearUAV(FRHICommandList& RHICmdList, FRHIUnorderedAccessView* Buffer, uint32 NumBytes, uint32 Value);
 
