Use cycles instead of duration timer

diff --git a/Engine/Source/Runtime/PhysicsSQ/Public/SQVerifier.h b/Engine/Source/Runtime/PhysicsSQ/Public/SQVerifier.h
--- a/Engine/Source/Runtime/PhysicsSQ/Public/SQVerifier.h
+++ b/Engine/Source/Runtime/PhysicsSQ/Public/SQVerifier.h
@@ -132,10 +132,10 @@ template <bool bHasPhysX = true>
 void SQPerfComparisonHelper(const FString& TestName, FPhysTestSerializer& Serializer, bool bEnsureOnMismatch = false)
 {
 	using namespace Chaos;
-	double PhysXSum = 0.0;
-	double ChaosSum = 0.0;
-	//double NumIterations = bHasPhysX ? 100 : 10000000000;
-	double NumIterations = 100;
+	uint32 PhysXSum = 0;
+	uint32 ChaosSum = 0;
+	double NumIterations = bHasPhysX ? 100 : 10000;
+	//double NumIterations = 1000000;
 
 	const FSQCapture& CapturedSQ = *Serializer.GetSQCapture();
 	switch (CapturedSQ.SQType)
@@ -147,10 +147,9 @@ void SQPerfComparisonHelper(const FString& TestName, FPhysTestSerializer& Serial
 			for (double i = 0; i < NumIterations; ++i)
 			{
 				auto PxHitBuffer = MakeUnique<PhysXInterface::FDynamicHitBuffer<PxRaycastHit>>();
-				FDurationTimer Timer(PhysXSum);
-				Timer.Start();
+				uint32 StartTime = FPlatformTime::Cycles();
 				Serializer.GetPhysXData()->raycast(U2PVector(CapturedSQ.StartPoint), U2PVector(CapturedSQ.Dir), CapturedSQ.DeltaMag, *PxHitBuffer, U2PHitFlags(CapturedSQ.OutputFlags.HitFlags), CapturedSQ.QueryFilterData, CapturedSQ.FilterCallback.Get());
-				Timer.Stop();
+				PhysXSum += FPlatformTime::Cycles() - StartTime;
 			}
 		}
 
@@ -160,10 +159,9 @@ void SQPerfComparisonHelper(const FString& TestName, FPhysTestSerializer& Serial
 		for (double i = 0; i < NumIterations; ++i)
 		{
 			auto ChaosHitBuffer = MakeUnique<ChaosInterface::FSQHitBuffer<ChaosInterface::FRaycastHit>>();
-			FDurationTimer Timer(ChaosSum);
-			Timer.Start();
+			uint32 StartTime = FPlatformTime::Cycles();
 			SQAccelerator.Raycast(CapturedSQ.StartPoint, CapturedSQ.Dir, CapturedSQ.DeltaMag, *ChaosHitBuffer, CapturedSQ.OutputFlags.HitFlags, CapturedSQ.QueryFilterData, *CapturedSQ.FilterCallback);
-			Timer.Stop();
+			ChaosSum += FPlatformTime::Cycles() - StartTime;
 		}
 		break;
 	}
@@ -174,10 +172,9 @@ void SQPerfComparisonHelper(const FString& TestName, FPhysTestSerializer& Serial
 			for (double i = 0; i < NumIterations; ++i)
 			{
 				auto PxHitBuffer = MakeUnique<PhysXInterface::FDynamicHitBuffer<PxSweepHit>>();
-				FDurationTimer Timer(PhysXSum);
-				Timer.Start();
+				uint32 StartTime = FPlatformTime::Cycles();
 				Serializer.GetPhysXData()->sweep(CapturedSQ.PhysXGeometry.any(), U2PTransform(CapturedSQ.StartTM), U2PVector(CapturedSQ.Dir), CapturedSQ.DeltaMag, *PxHitBuffer, U2PHitFlags(CapturedSQ.OutputFlags.HitFlags), CapturedSQ.QueryFilterData, CapturedSQ.FilterCallback.Get());
-				Timer.Stop();
+				PhysXSum += FPlatformTime::Cycles() - StartTime;
 			}
 		}
 
@@ -187,10 +184,9 @@ void SQPerfComparisonHelper(const FString& TestName, FPhysTestSerializer& Serial
 		for (double i = 0; i < NumIterations; ++i)
 		{
 			auto ChaosHitBuffer = MakeUnique<ChaosInterface::FSQHitBuffer<ChaosInterface::FSweepHit>>();
-			FDurationTimer Timer(ChaosSum);
-			Timer.Start();
+			uint32 StartTime = FPlatformTime::Cycles();
 			SQAccelerator.Sweep(*CapturedSQ.ChaosGeometry, CapturedSQ.StartTM, CapturedSQ.Dir, CapturedSQ.DeltaMag, *ChaosHitBuffer, CapturedSQ.OutputFlags.HitFlags, CapturedSQ.QueryFilterData, *CapturedSQ.FilterCallback);
-			Timer.Stop();
+			ChaosSum += FPlatformTime::Cycles() - StartTime;
 		}
 		break;
 	}
@@ -201,10 +197,9 @@ void SQPerfComparisonHelper(const FString& TestName, FPhysTestSerializer& Serial
 			for (double i = 0; i < NumIterations; ++i)
 			{
 				auto PxHitBuffer = MakeUnique<PhysXInterface::FDynamicHitBuffer<PxOverlapHit>>();
-				FDurationTimer Timer(PhysXSum);
-				Timer.Start();
+				uint32 StartTime = FPlatformTime::Cycles();
 				Serializer.GetPhysXData()->overlap(CapturedSQ.PhysXGeometry.any(), U2PTransform(CapturedSQ.StartTM), *PxHitBuffer, CapturedSQ.QueryFilterData, CapturedSQ.FilterCallback.Get());
-				Timer.Stop();
+				PhysXSum += FPlatformTime::Cycles() - StartTime;
 			}
 		}
 
@@ -214,25 +209,28 @@ void SQPerfComparisonHelper(const FString& TestName, FPhysTestSerializer& Serial
 		for (double i = 0; i < NumIterations; ++i)
 		{
 			auto ChaosHitBuffer = MakeUnique<ChaosInterface::FSQHitBuffer<ChaosInterface::FOverlapHit>>();
-			FDurationTimer Timer(ChaosSum);
-			Timer.Start();
+			uint32 StartTime = FPlatformTime::Cycles();
 			SQAccelerator.Overlap(*CapturedSQ.ChaosGeometry, CapturedSQ.StartTM, *ChaosHitBuffer, CapturedSQ.QueryFilterData, *CapturedSQ.FilterCallback);
-			Timer.Stop();
+			ChaosSum += FPlatformTime::Cycles() - StartTime;
 		}
 		break;
 	}
 	}
 
-	double AvgPhysX = 1000 * 1000 * PhysXSum / NumIterations;
-	double AvgChaos = 1000 * 1000 * ChaosSum / NumIterations;
+	float ChaosMs = FPlatformTime::ToMilliseconds(ChaosSum);
+	float ChaosAvgUs = (ChaosMs * 1000) / NumIterations;
+
+	float PhysXMs = FPlatformTime::ToMilliseconds(PhysXSum);
+	float PhysXAvgUs = (PhysXMs * 1000) / NumIterations;
 
 	if (bHasPhysX)
 	{
-		UE_LOG(LogPhysicsCore, Warning, TEXT("Perf Test:%s\nPhysX:%f(us), Chaos:%f(us)"), *TestName, AvgPhysX, AvgChaos);
+		UE_LOG(LogPhysicsCore, Warning, TEXT("Perf Test:%s\nPhysX:%f(us), Chaos:%f(us)"), *TestName, PhysXAvgUs, ChaosAvgUs);
 	}
 	else
 	{
-		UE_LOG(LogPhysicsCore, Warning, TEXT("Perf Test:%s\nChaos:%f(us)"), *TestName, AvgChaos);
+		UE_LOG(LogPhysicsCore, Warning, TEXT("Perf Test:%s\nChaos:%f(us), Total:%f(ms)"), *TestName, ChaosAvgUs, ChaosMs);
+		//UE_LOG(LogPhysicsCore, Warning, TEXT("Perf Test:%s\nChaos:%f(us)"), *TestName, AvgChaos);
 	}
 }
 
