Added a bounds check per shape when we have multiple shapes

diff --git a/Engine/Source/Runtime/Experimental/Chaos/Public/Chaos/Box.h b/Engine/Source/Runtime/Experimental/Chaos/Public/Chaos/Box.h
--- a/Engine/Source/Runtime/Experimental/Chaos/Public/Chaos/Box.h
+++ b/Engine/Source/Runtime/Experimental/Chaos/Public/Chaos/Box.h
@@ -257,6 +257,8 @@ namespace Chaos
 			return AABB.GetTypeHash();
 		}
 
+		const TAABB<T, d>& GetAABB() const { return AABB; }
+
 	private:
 		TAABB<T, d> AABB;
 	};
diff --git a/Engine/Source/Runtime/PhysicsSQ/Private/SQAccelerator.cpp b/Engine/Source/Runtime/PhysicsSQ/Private/SQAccelerator.cpp
--- a/Engine/Source/Runtime/PhysicsSQ/Private/SQAccelerator.cpp
+++ b/Engine/Source/Runtime/PhysicsSQ/Private/SQAccelerator.cpp
@@ -92,6 +92,7 @@ struct TSQVisitor : public Chaos::ISpatialVisitor<TPayload, float>
 		, QueryFilterData(InQueryFilterData)
 		, QueryCallback(InQueryCallback)
 		, bAnyHit(false)
+		, HalfExtents(0)
 		, DebugParams(InDebugParams)
 	{
 #if WITH_PHYSX
@@ -110,6 +111,7 @@ struct TSQVisitor : public Chaos::ISpatialVisitor<TPayload, float>
 		, QueryCallback(InQueryCallback)
 		, bAnyHit(false)
 		, QueryGeom(&InQueryGeom)
+		, HalfExtents(InQueryGeom.BoundingBox().Extents() * 0.5)
 		, DebugParams(InDebugParams)
 	{
 #if WITH_PHYSX
@@ -127,6 +129,7 @@ struct TSQVisitor : public Chaos::ISpatialVisitor<TPayload, float>
 		, QueryCallback(InQueryCallback)
 		, bAnyHit(false)
 		, QueryGeom(&InQueryGeom)
+		, HalfExtents(InQueryGeom.BoundingBox().Extents() * 0.5)
 		, DebugParams(InDebugParams)
 	{
 #if WITH_PHYSX
@@ -177,9 +180,37 @@ private:
 		TGeometryParticle<float, 3>* GeometryParticle = Payload.GetExternalGeometryParticle_ExternalThread();
 		const TShapesArray<float,3>& Shapes = GeometryParticle->ShapesArray();
 
+		const bool bTestShapeBounds = Shapes.Num() > 1;
+		const TRigidTransform<float, 3> ActorTM(GeometryParticle->X(), GeometryParticle->R());
+
 		for (const auto& Shape : Shapes)
 		{
 			const FImplicitObject* Geom = Shape->Geometry.Get();
+
+			if (bTestShapeBounds)
+			{
+				const auto& GeomBounds = Geom->BoundingBox();
+				auto InflatedWorldBounds = SQ == ESQType::Raycast ? GeomBounds.GetAABB().TransformedAABB(ActorTM) :
+					TAABB<FReal, 3>(GeomBounds.Min() - HalfExtents, GeomBounds.Max() + HalfExtents).TransformedAABB(ActorTM);
+				if (SQ != ESQType::Overlap)
+				{
+					//todo: use fast raycast
+					float TmpTime;
+					FVec3 TmpPos;
+					if (!InflatedWorldBounds.RaycastFast( SQ == ESQType::Raycast ? StartPoint : StartTM.GetLocation(), CurData->Dir, CurData->InvDir, CurData->bParallel, CurData->CurrentLength, CurData->InvCurrentLength, TmpTime, TmpPos))
+					{
+						continue;
+					}
+				}
+				else
+				{
+					if (!InflatedWorldBounds.Contains(StartTM.GetLocation()))
+					{
+						continue;
+					}
+				}
+			}
+
 			//TODO: use gt particles directly
 			//#TODO alternative to px flags
 #if WITH_PHYSX
@@ -191,7 +222,6 @@ private:
 #endif
 			if (HitType != ECollisionQueryHitType::None)
 			{
-				const TRigidTransform<float, 3> ActorTM(GeometryParticle->X(), GeometryParticle->R());
 
 				THitType Hit;
 				Hit.Actor = GeometryParticle;
@@ -290,6 +320,7 @@ private:
 	ICollisionQueryFilterCallbackBase& QueryCallback;
 	bool bAnyHit;
 	const QueryGeometryType* QueryGeom;
+	const FVector HalfExtents;
 	const FQueryDebugParams DebugParams;
 };
 
