Disabling shape bounds optimization temporarily due to conflicts with some shape unions causing SQ to fail to check relevant shapes

diff --git a/Engine/Source/Runtime/PhysicsSQ/Private/SQAccelerator.cpp b/Engine/Source/Runtime/PhysicsSQ/Private/SQAccelerator.cpp
--- a/Engine/Source/Runtime/PhysicsSQ/Private/SQAccelerator.cpp
+++ b/Engine/Source/Runtime/PhysicsSQ/Private/SQAccelerator.cpp
@@ -180,7 +180,11 @@ private:
 		TGeometryParticle<float, 3>* GeometryParticle = Payload.GetExternalGeometryParticle_ExternalThread();
 		const TShapesArray<float,3>& Shapes = GeometryParticle->ShapesArray();
 
-		const bool bTestShapeBounds = Shapes.Num() > 1;
+		//////////////////////////////////////////////////////////////////////////
+		// Temporarily disabled shape bounds checked due to conflicts with some shape unions
+		const bool bTestShapeBounds = false; // Shapes.Num() > 1;
+		//////////////////////////////////////////////////////////////////////////
+
 		const TRigidTransform<float, 3> ActorTM(GeometryParticle->X(), GeometryParticle->R());
 
 		for (const auto& Shape : Shapes)
