Fix ensure when sweep of length 0 is used. Add proper support so that we get back initial overlaps

diff --git a/Engine/Source/Runtime/PhysicsSQ/Private/SQAccelerator.cpp b/Engine/Source/Runtime/PhysicsSQ/Private/SQAccelerator.cpp
--- a/Engine/Source/Runtime/PhysicsSQ/Private/SQAccelerator.cpp
+++ b/Engine/Source/Runtime/PhysicsSQ/Private/SQAccelerator.cpp
@@ -92,6 +92,7 @@ struct TSQVisitor : public Chaos::ISpatialVisitor<TPayload, float>
 		, QueryFilterData(InQueryFilterData)
 		, QueryCallback(InQueryCallback)
 		, bAnyHit(false)
+		, bOverlapCameFromSweep(false)
 		, HalfExtents(0)
 		, DebugParams(InDebugParams)
 	{
@@ -102,7 +103,7 @@ struct TSQVisitor : public Chaos::ISpatialVisitor<TPayload, float>
 	}
 
 	TSQVisitor(const FTransform& InStartTM, const FVector& InDir, ChaosInterface::FSQHitBuffer<ChaosInterface::FSweepHit>& InHitBuffer, EHitFlags InOutputFlags,
-		const FQueryFilterData& InQueryFilterData, ICollisionQueryFilterCallbackBase& InQueryCallback, const QueryGeometryType& InQueryGeom, const FQueryDebugParams& InDebugParams)
+		const FQueryFilterData& InQueryFilterData, ICollisionQueryFilterCallbackBase& InQueryCallback, const QueryGeometryType& InQueryGeom, bool bSweepAsOverlap, const FQueryDebugParams& InDebugParams)
 		: StartTM(InStartTM)
 		, Dir(InDir)
 		, HitBuffer(InHitBuffer)
@@ -110,6 +111,7 @@ struct TSQVisitor : public Chaos::ISpatialVisitor<TPayload, float>
 		, QueryFilterData(InQueryFilterData)
 		, QueryCallback(InQueryCallback)
 		, bAnyHit(false)
+		, bOverlapCameFromSweep(bSweepAsOverlap)
 		, QueryGeom(&InQueryGeom)
 		, HalfExtents(InQueryGeom.BoundingBox().Extents() * 0.5)
 		, DebugParams(InDebugParams)
@@ -278,15 +280,28 @@ private:
 					{
 
 						//overlap never blocks
-						const bool bBlocker = SQ != ESQType::Overlap && (HitType == ECollisionQueryHitType::Block || bAnyHit || HitBuffer.WantsSingleResult());
+						const bool bIsTrulyOverlap = SQ == ESQType::Overlap && !bOverlapCameFromSweep;
+						const bool bBlocker = !bIsTrulyOverlap && (HitType == ECollisionQueryHitType::Block || bAnyHit || HitBuffer.WantsSingleResult());
 						HitBuffer.InsertHit(Hit, bBlocker);
 
 						if (bBlocker)
 						{
-							CurData->SetLength(FMath::Max(0.f, Distance));	//Max is needed for MTD which returns negative distance
-							if (CurData->CurrentLength == 0 && (SQ == ESQType::Raycast || HitBuffer.WantsSingleResult()))	//raycasts always fail with distance 0, sweeps only matter if we want multi overlaps
+							if (SQ != ESQType::Overlap)
+							{
+								CurData->SetLength(FMath::Max(0.f, Distance));	//Max is needed for MTD which returns negative distance
+								if (CurData->CurrentLength == 0 && (SQ == ESQType::Raycast || HitBuffer.WantsSingleResult()))	//raycasts always fail with distance 0, sweeps only matter if we want multi overlaps
+								{
+									return false;	//initial overlap so nothing will be better than this
+								}
+							}
+							else
 							{
-								return false;	//initial overlap so nothing will be better than this
+								//must be sweep as overlap, in that case first blocker is enough unless they want overlap multi
+								ensure(bOverlapCameFromSweep);
+								if (HitBuffer.WantsSingleResult())
+								{
+									return false;
+								}
 							}
 						}
 
@@ -327,6 +342,7 @@ private:
 	const FQueryFilterData& QueryFilterData;
 	ICollisionQueryFilterCallbackBase& QueryCallback;
 	bool bAnyHit;
+	bool bOverlapCameFromSweep;
 	const QueryGeometryType* QueryGeom;
 	const FVector HalfExtents;
 	const FQueryDebugParams DebugParams;
@@ -350,11 +366,22 @@ void SweepHelper(const QueryGeomType& QueryGeom, const Chaos::ISpatialAccelerati
 	using namespace ChaosInterface;
 
 	const TAABB<float, 3> Bounds = QueryGeom.BoundingBox().TransformedAABB(StartTM);
-	const FVector HalfExtents = Bounds.Extents() * 0.5f;
+	const bool bSweepAsOverlap = DeltaMagnitude == 0;	//question: do we care about tiny sweeps?
+	TSQVisitor<QueryGeomType, TAccelerationStructureHandle<float, 3>, FSweepHit> SweepVisitor(StartTM, Dir, HitBuffer, OutputFlags, QueryFilterData, QueryCallback, QueryGeom, bSweepAsOverlap, DebugParams);
 
-	TSQVisitor<QueryGeomType, TAccelerationStructureHandle<float, 3>, FSweepHit> SweepVisitor(StartTM, Dir, HitBuffer, OutputFlags, QueryFilterData, QueryCallback, QueryGeom, DebugParams);
 	HitBuffer.IncFlushCount();
-	SpatialAcceleration.Sweep(Bounds.GetCenter(), Dir, DeltaMagnitude, HalfExtents, SweepVisitor);
+
+	if (bSweepAsOverlap)
+	{
+		//fallback to overlap
+		SpatialAcceleration.Overlap(Bounds, SweepVisitor);
+	}
+	else
+	{
+		const FVector HalfExtents = Bounds.Extents() * 0.5f;
+		SpatialAcceleration.Sweep(Bounds.GetCenter(), Dir, DeltaMagnitude, HalfExtents, SweepVisitor);
+	}
+
 	HitBuffer.DecFlushCount();
 }
 
