Back out changelist 10911205. EdTest errors

diff --git a/Engine/Source/Runtime/PhysicsSQ/Private/SQAccelerator.cpp b/Engine/Source/Runtime/PhysicsSQ/Private/SQAccelerator.cpp
--- a/Engine/Source/Runtime/PhysicsSQ/Private/SQAccelerator.cpp
+++ b/Engine/Source/Runtime/PhysicsSQ/Private/SQAccelerator.cpp
@@ -92,7 +92,6 @@ struct TSQVisitor : public Chaos::ISpatialVisitor<TPayload, float>
 		, QueryFilterData(InQueryFilterData)
 		, QueryCallback(InQueryCallback)
 		, bAnyHit(false)
-		, bOverlapCameFromSweep(false)
 		, HalfExtents(0)
 		, DebugParams(InDebugParams)
 	{
@@ -103,7 +102,7 @@ struct TSQVisitor : public Chaos::ISpatialVisitor<TPayload, float>
 	}
 
 	TSQVisitor(const FTransform& InStartTM, const FVector& InDir, ChaosInterface::FSQHitBuffer<ChaosInterface::FSweepHit>& InHitBuffer, EHitFlags InOutputFlags,
-		const FQueryFilterData& InQueryFilterData, ICollisionQueryFilterCallbackBase& InQueryCallback, const QueryGeometryType& InQueryGeom, bool bSweepAsOverlap, const FQueryDebugParams& InDebugParams)
+		const FQueryFilterData& InQueryFilterData, ICollisionQueryFilterCallbackBase& InQueryCallback, const QueryGeometryType& InQueryGeom, const FQueryDebugParams& InDebugParams)
 		: StartTM(InStartTM)
 		, Dir(InDir)
 		, HitBuffer(InHitBuffer)
@@ -111,7 +110,6 @@ struct TSQVisitor : public Chaos::ISpatialVisitor<TPayload, float>
 		, QueryFilterData(InQueryFilterData)
 		, QueryCallback(InQueryCallback)
 		, bAnyHit(false)
-		, bOverlapCameFromSweep(bSweepAsOverlap)
 		, QueryGeom(&InQueryGeom)
 		, HalfExtents(InQueryGeom.BoundingBox().Extents() * 0.5)
 		, DebugParams(InDebugParams)
@@ -280,28 +278,15 @@ private:
 					{
 
 						//overlap never blocks
-						const bool bIsTrulyOverlap = SQ == ESQType::Overlap && !bOverlapCameFromSweep;
-						const bool bBlocker = !bIsTrulyOverlap && (HitType == ECollisionQueryHitType::Block || bAnyHit || HitBuffer.WantsSingleResult());
+						const bool bBlocker = SQ != ESQType::Overlap && (HitType == ECollisionQueryHitType::Block || bAnyHit || HitBuffer.WantsSingleResult());
 						HitBuffer.InsertHit(Hit, bBlocker);
 
 						if (bBlocker)
 						{
-							if (SQ != ESQType::Overlap)
-							{
-								CurData->SetLength(FMath::Max(0.f, Distance));	//Max is needed for MTD which returns negative distance
-								if (CurData->CurrentLength == 0 && (SQ == ESQType::Raycast || HitBuffer.WantsSingleResult()))	//raycasts always fail with distance 0, sweeps only matter if we want multi overlaps
-								{
-									return false;	//initial overlap so nothing will be better than this
-								}
-							}
-							else
+							CurData->SetLength(FMath::Max(0.f, Distance));	//Max is needed for MTD which returns negative distance
+							if (CurData->CurrentLength == 0 && (SQ == ESQType::Raycast || HitBuffer.WantsSingleResult()))	//raycasts always fail with distance 0, sweeps only matter if we want multi overlaps
 							{
-								//must be sweep as overlap, in that case first blocker is enough unless they want overlap multi
-								ensure(bOverlapCameFromSweep);
-								if (HitBuffer.WantsSingleResult())
-								{
-									return false;
-								}
+								return false;	//initial overlap so nothing will be better than this
 							}
 						}
 
@@ -342,7 +327,6 @@ private:
 	const FQueryFilterData& QueryFilterData;
 	ICollisionQueryFilterCallbackBase& QueryCallback;
 	bool bAnyHit;
-	bool bOverlapCameFromSweep;
 	const QueryGeometryType* QueryGeom;
 	const FVector HalfExtents;
 	const FQueryDebugParams DebugParams;
@@ -366,22 +350,11 @@ void SweepHelper(const QueryGeomType& QueryGeom, const Chaos::ISpatialAccelerati
 	using namespace ChaosInterface;
 
 	const TAABB<float, 3> Bounds = QueryGeom.BoundingBox().TransformedAABB(StartTM);
-	const bool bSweepAsOverlap = DeltaMagnitude == 0;	//question: do we care about tiny sweeps?
-	TSQVisitor<QueryGeomType, TAccelerationStructureHandle<float, 3>, FSweepHit> SweepVisitor(StartTM, Dir, HitBuffer, OutputFlags, QueryFilterData, QueryCallback, QueryGeom, bSweepAsOverlap, DebugParams);
+	const FVector HalfExtents = Bounds.Extents() * 0.5f;
 
+	TSQVisitor<QueryGeomType, TAccelerationStructureHandle<float, 3>, FSweepHit> SweepVisitor(StartTM, Dir, HitBuffer, OutputFlags, QueryFilterData, QueryCallback, QueryGeom, DebugParams);
 	HitBuffer.IncFlushCount();
-
-	if (bSweepAsOverlap)
-	{
-		//fallback to overlap
-		SpatialAcceleration.Overlap(Bounds, SweepVisitor);
-	}
-	else
-	{
-		const FVector HalfExtents = Bounds.Extents() * 0.5f;
-		SpatialAcceleration.Sweep(Bounds.GetCenter(), Dir, DeltaMagnitude, HalfExtents, SweepVisitor);
-	}
-
+	SpatialAcceleration.Sweep(Bounds.GetCenter(), Dir, DeltaMagnitude, HalfExtents, SweepVisitor);
 	HitBuffer.DecFlushCount();
 }
 
