Realign TSQVisitor<> member layout such that size decreases from 160 bytes per instance with 32 bytes padding down to class size of 128 bytes per instance with 0 bytes padding.

diff --git a/Engine/Source/Runtime/PhysicsSQ/Private/SQAccelerator.cpp b/Engine/Source/Runtime/PhysicsSQ/Private/SQAccelerator.cpp
--- a/Engine/Source/Runtime/PhysicsSQ/Private/SQAccelerator.cpp
+++ b/Engine/Source/Runtime/PhysicsSQ/Private/SQAccelerator.cpp
@@ -87,13 +87,13 @@ struct TSQVisitor : public Chaos::ISpatialVisitor<TPayload, float>
 		const FQueryFilterData& InQueryFilterData, ICollisionQueryFilterCallbackBase& InQueryCallback, const FQueryDebugParams& InDebugParams)
 		: StartPoint(InStartPoint)
 		, Dir(InDir)
-		, HitBuffer(InHitBuffer)
+		, HalfExtents(0)
 		, OutputFlags(InOutputFlags)
-		, QueryFilterData(InQueryFilterData)
-		, QueryCallback(InQueryCallback)
 		, bAnyHit(false)
-		, HalfExtents(0)
 		, DebugParams(InDebugParams)
+		, HitBuffer(InHitBuffer)
+		, QueryFilterData(InQueryFilterData)
+		, QueryCallback(InQueryCallback)
 	{
 #if WITH_PHYSX
 		//#TODO - reimplement query flags alternative for Chaos
@@ -103,16 +103,16 @@ struct TSQVisitor : public Chaos::ISpatialVisitor<TPayload, float>
 
 	TSQVisitor(const FTransform& InStartTM, const FVector& InDir, ChaosInterface::FSQHitBuffer<ChaosInterface::FSweepHit>& InHitBuffer, EHitFlags InOutputFlags,
 		const FQueryFilterData& InQueryFilterData, ICollisionQueryFilterCallbackBase& InQueryCallback, const QueryGeometryType& InQueryGeom, const FQueryDebugParams& InDebugParams)
-		: StartTM(InStartTM)
-		, Dir(InDir)
-		, HitBuffer(InHitBuffer)
+		: Dir(InDir)
+		, HalfExtents(InQueryGeom.BoundingBox().Extents() * 0.5)
 		, OutputFlags(InOutputFlags)
-		, QueryFilterData(InQueryFilterData)
-		, QueryCallback(InQueryCallback)
 		, bAnyHit(false)
-		, QueryGeom(&InQueryGeom)
-		, HalfExtents(InQueryGeom.BoundingBox().Extents() * 0.5)
 		, DebugParams(InDebugParams)
+		, HitBuffer(InHitBuffer)
+		, QueryFilterData(InQueryFilterData)
+		, QueryGeom(&InQueryGeom)
+		, QueryCallback(InQueryCallback)
+		, StartTM(InStartTM)
 	{
 #if WITH_PHYSX
 		//#TODO - reimplement query flags alternative for Chaos
@@ -123,14 +123,14 @@ struct TSQVisitor : public Chaos::ISpatialVisitor<TPayload, float>
 
 	TSQVisitor(const FTransform& InWorldTM, ChaosInterface::FSQHitBuffer<ChaosInterface::FOverlapHit>& InHitBuffer,
 		const FQueryFilterData& InQueryFilterData, ICollisionQueryFilterCallbackBase& InQueryCallback, const QueryGeometryType& InQueryGeom, const FQueryDebugParams& InDebugParams)
-		: StartTM(InWorldTM)
+		: HalfExtents(InQueryGeom.BoundingBox().Extents() * 0.5)
+		, bAnyHit(false)
+		, DebugParams(InDebugParams)
 		, HitBuffer(InHitBuffer)
 		, QueryFilterData(InQueryFilterData)
-		, QueryCallback(InQueryCallback)
-		, bAnyHit(false)
 		, QueryGeom(&InQueryGeom)
-		, HalfExtents(InQueryGeom.BoundingBox().Extents() * 0.5)
-		, DebugParams(InDebugParams)
+		, QueryCallback(InQueryCallback)
+		, StartTM(InWorldTM)
 	{
 #if WITH_PHYSX
 		//#TODO - reimplement query flags alternative for Chaos
@@ -319,17 +319,17 @@ private:
 	}
 #endif
 
-	const FTransform StartTM;
 	const FVector StartPoint;
 	const FVector Dir;
-	ChaosInterface::FSQHitBuffer<THitType>& HitBuffer;
+	const FVector HalfExtents;
 	EHitFlags OutputFlags;
-	const FQueryFilterData& QueryFilterData;
-	ICollisionQueryFilterCallbackBase& QueryCallback;
 	bool bAnyHit;
-	const QueryGeometryType* QueryGeom;
-	const FVector HalfExtents;
 	const FQueryDebugParams DebugParams;
+	ChaosInterface::FSQHitBuffer<THitType>& HitBuffer;
+	const FQueryFilterData& QueryFilterData;
+	const QueryGeometryType* QueryGeom;
+	ICollisionQueryFilterCallbackBase& QueryCallback;
+	const FTransform StartTM;
 };
 
 void FChaosSQAccelerator::Raycast(const FVector& Start, const FVector& Dir, const float DeltaMagnitude, ChaosInterface::FSQHitBuffer<ChaosInterface::FRaycastHit>& HitBuffer, EHitFlags OutputFlags, const FQueryFilterData& QueryFilterData, ICollisionQueryFilterCallbackBase& QueryCallback, const FQueryDebugParams& DebugParams) const
