From Ori: fix logic bug that was causing nothing to be returned in overlap queries that were treated as "any hit".

#rnx

diff --git a/Engine/Source/Runtime/PhysicsSQ/Private/SQAccelerator.cpp b/Engine/Source/Runtime/PhysicsSQ/Private/SQAccelerator.cpp
--- a/Engine/Source/Runtime/PhysicsSQ/Private/SQAccelerator.cpp
+++ b/Engine/Source/Runtime/PhysicsSQ/Private/SQAccelerator.cpp
@@ -278,10 +278,10 @@ private:
 					{
 
 						//overlap never blocks
-						const bool bBlocker = SQ != ESQType::Overlap && (HitType == ECollisionQueryHitType::Block || bAnyHit || HitBuffer.WantsSingleResult());
+						const bool bBlocker = (HitType == ECollisionQueryHitType::Block || bAnyHit || HitBuffer.WantsSingleResult());
 						HitBuffer.InsertHit(Hit, bBlocker);
 
-						if (bBlocker)
+						if (bBlocker && SQ != ESQType::Overlap)
 						{
 							CurData->SetLength(FMath::Max(0.f, Distance));	//Max is needed for MTD which returns negative distance
 							if (CurData->CurrentLength == 0 && (SQ == ESQType::Raycast || HitBuffer.WantsSingleResult()))	//raycasts always fail with distance 0, sweeps only matter if we want multi overlaps
