Reenable fix for sweeps of length 0 ignoring objects

diff --git a/Engine/Source/Runtime/PhysicsSQ/Private/SQAccelerator.cpp b/Engine/Source/Runtime/PhysicsSQ/Private/SQAccelerator.cpp
--- a/Engine/Source/Runtime/PhysicsSQ/Private/SQAccelerator.cpp
+++ b/Engine/Source/Runtime/PhysicsSQ/Private/SQAccelerator.cpp
@@ -344,17 +344,27 @@ void FChaosSQAccelerator::Raycast(const FVector& Start, const FVector& Dir, cons
 }
 
 template <typename QueryGeomType>
-void SweepHelper(const QueryGeomType& QueryGeom, const Chaos::ISpatialAcceleration<Chaos::TAccelerationStructureHandle<float, 3>, float, 3>& SpatialAcceleration, const FTransform& StartTM, const FVector& Dir, const float DeltaMagnitude, ChaosInterface::FSQHitBuffer<ChaosInterface::FSweepHit>& HitBuffer, EHitFlags OutputFlags, const FQueryFilterData& QueryFilterData, ICollisionQueryFilterCallbackBase& QueryCallback, const FQueryDebugParams& DebugParams)
+void SweepHelper(const QueryGeomType& QueryGeom,const Chaos::ISpatialAcceleration<Chaos::TAccelerationStructureHandle<float,3>,float,3>& SpatialAcceleration,const FTransform& StartTM,const FVector& Dir,const float DeltaMagnitude,ChaosInterface::FSQHitBuffer<ChaosInterface::FSweepHit>& HitBuffer,EHitFlags OutputFlags,const FQueryFilterData& QueryFilterData,ICollisionQueryFilterCallbackBase& QueryCallback,const FQueryDebugParams& DebugParams)
 {
 	using namespace Chaos;
 	using namespace ChaosInterface;
 
-	const TAABB<float, 3> Bounds = QueryGeom.BoundingBox().TransformedAABB(StartTM);
-	const FVector HalfExtents = Bounds.Extents() * 0.5f;
+	const TAABB<float,3> Bounds = QueryGeom.BoundingBox().TransformedAABB(StartTM);
+	const bool bSweepAsOverlap = DeltaMagnitude == 0;	//question: do we care about tiny sweeps?
+	TSQVisitor<QueryGeomType,TAccelerationStructureHandle<float,3>,FSweepHit> SweepVisitor(StartTM,Dir,HitBuffer,OutputFlags,QueryFilterData,QueryCallback,QueryGeom,DebugParams);
 
-	TSQVisitor<QueryGeomType, TAccelerationStructureHandle<float, 3>, FSweepHit> SweepVisitor(StartTM, Dir, HitBuffer, OutputFlags, QueryFilterData, QueryCallback, QueryGeom, DebugParams);
 	HitBuffer.IncFlushCount();
-	SpatialAcceleration.Sweep(Bounds.GetCenter(), Dir, DeltaMagnitude, HalfExtents, SweepVisitor);
+
+	if(bSweepAsOverlap)
+	{
+		//fallback to overlap
+		SpatialAcceleration.Overlap(Bounds,SweepVisitor);
+	} else
+	{
+		const FVector HalfExtents = Bounds.Extents() * 0.5f;
+		SpatialAcceleration.Sweep(Bounds.GetCenter(),Dir,DeltaMagnitude,HalfExtents,SweepVisitor);
+	}
+
 	HitBuffer.DecFlushCount();
 }
 
