As the halfextent was extrated from the untransformed shape we need to transform that extent before we check against it.

diff --git a/Engine/Source/Runtime/PhysicsSQ/Private/SQAccelerator.cpp b/Engine/Source/Runtime/PhysicsSQ/Private/SQAccelerator.cpp
--- a/Engine/Source/Runtime/PhysicsSQ/Private/SQAccelerator.cpp
+++ b/Engine/Source/Runtime/PhysicsSQ/Private/SQAccelerator.cpp
@@ -207,8 +207,10 @@ private:
 				}
 				else
 				{
-					InflatedWorldBounds = TAABB<FReal, 3>(Shape->WorldSpaceInflatedShapeBounds.Min() - HalfExtents, Shape->WorldSpaceInflatedShapeBounds.Max() + HalfExtents);
+					// Transform to world bounds and get the proper half extent.
+					const FVec3 WorldHalfExtent = QueryGeom ? QueryGeom->BoundingBox().TransformedAABB(StartTM).Extents() * 0.5f : HalfExtents;
 
+					InflatedWorldBounds = TAABB<FReal, 3>(Shape->WorldSpaceInflatedShapeBounds.Min() - WorldHalfExtent, Shape->WorldSpaceInflatedShapeBounds.Max() + WorldHalfExtent);
 				}
 				if (SQ != ESQType::Overlap)
 				{
