Fix static analysis warnings

diff --git a/Engine/Source/Runtime/Experimental/Chaos/Private/Chaos/CollisionResolution.cpp b/Engine/Source/Runtime/Experimental/Chaos/Private/Chaos/CollisionResolution.cpp
--- a/Engine/Source/Runtime/Experimental/Chaos/Private/Chaos/CollisionResolution.cpp
+++ b/Engine/Source/Runtime/Experimental/Chaos/Private/Chaos/CollisionResolution.cpp
@@ -2338,6 +2338,10 @@ namespace Chaos
 			}
 			Dir /= Length;
 
+			// Attempt to get rid of invalid static analysis warnings
+			check(Implicit0 || Implicit0Type == ImplicitObjectType::Unknown);
+			check(Implicit1 || Implicit1Type == ImplicitObjectType::Unknown);
+
 			if (Implicit0Type == TBox<FReal, 3>::StaticType() && Implicit1Type == THeightField<FReal>::StaticType())
 			{
 				return GetPairTOIHackImpl<TBox<FReal, 3>>(*Implicit0->template GetObject<TBox<FReal, 3>>(), StartTransform0, *Implicit1->template GetObject<THeightField<FReal>>(), Transform1, Dir, Length, OutTOI, OutNormal, OutPhi);
diff --git a/Engine/Source/Runtime/PhysicsSQ/Private/SQAccelerator.cpp b/Engine/Source/Runtime/PhysicsSQ/Private/SQAccelerator.cpp
--- a/Engine/Source/Runtime/PhysicsSQ/Private/SQAccelerator.cpp
+++ b/Engine/Source/Runtime/PhysicsSQ/Private/SQAccelerator.cpp
@@ -268,11 +268,11 @@ private:
 						WorldNormal = ActorTM.TransformVectorNoScale(LocalNormal);
 					}
 				}
-				else if(SQ == ESQType::Sweep && CurData->CurrentLength > 0)
+				else if(SQ == ESQType::Sweep && CurData->CurrentLength > 0 && ensure(QueryGeom))
 				{
 					bHit = SweepQuery(*Geom, ActorTM, *QueryGeom, StartTM, CurData->Dir, CurData->CurrentLength, Distance, WorldPosition, WorldNormal, FaceIdx, 0.f, bComputeMTD);
 				}
-				else if (SQ == ESQType::Overlap || (SQ == ESQType::Sweep && CurData->CurrentLength == 0))
+				else if ((SQ == ESQType::Overlap || (SQ == ESQType::Sweep && CurData->CurrentLength == 0)) && ensure(QueryGeom))
 				{
 					if (bComputeMTD)
 					{
