Remove PhysX module when using Chaos, Guard PhysX code in PHYSICS_INTERFACE_PHYSX so Chaos compiles without it.
Change WITH_PHYSX to always be defined, even with chaos only.
Make Client/Server targets only disable PhysX when Chaos is enabled.

diff --git a/Engine/Plugins/Experimental/ChaosCloth/Source/Chaos/Private/ChaosCloth/ChaosClothingSimulation.cpp b/Engine/Plugins/Experimental/ChaosCloth/Source/Chaos/Private/ChaosCloth/ChaosClothingSimulation.cpp
--- a/Engine/Plugins/Experimental/ChaosCloth/Source/Chaos/Private/ChaosCloth/ChaosClothingSimulation.cpp
+++ b/Engine/Plugins/Experimental/ChaosCloth/Source/Chaos/Private/ChaosCloth/ChaosClothingSimulation.cpp
@@ -40,7 +40,7 @@
 #include "Chaos/XPBDSpringConstraints.h"
 #include "Chaos/XPBDAxialSpringConstraints.h"
 
-#if WITH_PHYSX && !PLATFORM_LUMIN && !PLATFORM_ANDROID
+#if PHYSICS_INTERFACE_PHYSX && !PLATFORM_LUMIN && !PLATFORM_ANDROID
 #include "PhysXIncludes.h"
 #endif
 
@@ -932,7 +932,7 @@ void ClothingSimulation::ExtractPhysicsAssetCollisions(const UClothingAssetCommo
 				// Add stub for extracted collision data
 				FClothCollisionPrim_Convex Convex;
 				Convex.BoneIndex = MappedBoneIndex;
-#if WITH_PHYSX
+#if PHYSICS_INTERFACE_PHYSX
 				// Collision bodies are stored in PhysX specific data structures so they can only be imported if we enable PhysX.
 				const physx::PxConvexMesh* const PhysXMesh = ConvexElem.GetConvexMesh();  // TODO(Kriss.Gossart): Deal with this legacy structure in a different place, so that there's only TConvex
 				const int32 NumPolygons = int32(PhysXMesh->getNbPolygons());
@@ -953,7 +953,7 @@ void ClothingSimulation::ExtractPhysicsAssetCollisions(const UClothingAssetCommo
 				// Rebuild surface points
 				Convex.RebuildSurfacePoints();
 
-#elif WITH_CHAOS  // #if WITH_PHYSX
+#elif WITH_CHAOS  // #if PHYSICS_INTERFACE_PHYSX
 				const Chaos::FImplicitObject& ChaosConvexMesh = *ConvexElem.GetChaosConvexMesh();
 				const Chaos::FConvex& ChaosConvex = ChaosConvexMesh.GetObjectChecked<Chaos::FConvex>();
 
@@ -972,7 +972,7 @@ void ClothingSimulation::ExtractPhysicsAssetCollisions(const UClothingAssetCommo
 				{
 					Convex.SurfacePoints.Add(ChaosConvex.GetSurfaceParticles().X(ParticleIndex));
 				}
-#endif  // #if WITH_PHYSX #elif WITH_CHAOS
+#endif  // #if PHYSICS_INTERFACE_PHYSX #elif WITH_CHAOS
 
 				// Add extracted collision data
 				ExtractedCollisions.Convexes.Add(Convex);
diff --git a/Engine/Plugins/Runtime/ApexDestruction/Source/ApexDestruction/Private/DestructibleComponent.cpp b/Engine/Plugins/Runtime/ApexDestruction/Source/ApexDestruction/Private/DestructibleComponent.cpp
--- a/Engine/Plugins/Runtime/ApexDestruction/Source/ApexDestruction/Private/DestructibleComponent.cpp
+++ b/Engine/Plugins/Runtime/ApexDestruction/Source/ApexDestruction/Private/DestructibleComponent.cpp
@@ -26,13 +26,13 @@
 #include "Physics/PhysicsInterfaceUtils.h"
 #include "UObject/UObjectThreadContext.h"
 
-#if WITH_PHYSX
+#if PHYSICS_INTERFACE_PHYSX
 #include "Physics/PhysicsGeometry.h"
 #endif
 
 UDestructibleComponent::UDestructibleComponent(const FObjectInitializer& ObjectInitializer)
 	: Super(ObjectInitializer)
-#if WITH_PHYSX
+#if PHYSICS_INTERFACE_PHYSX
 	, PhysxUserData(this)
 #endif
 {
@@ -1449,7 +1449,7 @@ bool UDestructibleComponent::IsAnySimulatingPhysics() const
 	return !!BodyInstance.bSimulatePhysics;
 }
 
-#if WITH_PHYSX
+#if PHYSICS_INTERFACE_PHYSX
 
 bool UDestructibleComponent::IsChunkLarge(PxRigidActor* ChunkActor) const
 {
@@ -1534,7 +1534,7 @@ void UDestructibleComponent::SetCollisionProfileName(FName InCollisionProfileNam
 }
 
 
-#if WITH_PHYSX
+#if PHYSICS_INTERFACE_PHYSX
 
 void UDestructibleComponent::SetCollisionResponseForActor(PxRigidDynamic* Actor, int32 ChunkIdx, const FCollisionResponseContainer* ResponseOverride /*= NULL*/)
 {
diff --git a/Engine/Plugins/Runtime/PhysXVehicles/Source/PhysXVehicles/Private/WheeledVehicleMovementComponent.cpp b/Engine/Plugins/Runtime/PhysXVehicles/Source/PhysXVehicles/Private/WheeledVehicleMovementComponent.cpp
--- a/Engine/Plugins/Runtime/PhysXVehicles/Source/PhysXVehicles/Private/WheeledVehicleMovementComponent.cpp
+++ b/Engine/Plugins/Runtime/PhysXVehicles/Source/PhysXVehicles/Private/WheeledVehicleMovementComponent.cpp
@@ -33,8 +33,7 @@
 #define LOCTEXT_NAMESPACE "UWheeledVehicleMovementComponent"
 
 
-#if WITH_PHYSX
-
+#if PHYSICS_INTERFACE_PHYSX
 /**
  * PhysX shader for tire friction forces
  * tireFriction - friction value of the tire contact.
@@ -141,7 +140,7 @@ UWheeledVehicleMovementComponent::UWheeledVehicleMovementComponent(const FObject
 	
 	bReverseAsBrake = true;	//Treats reverse button as break for a more arcade feel (also automatically goes into reverse)
 
-#if WITH_PHYSX
+#if PHYSICS_INTERFACE_PHYSX
 	// tire load filtering
 	PxVehicleTireLoadFilterData PTireLoadFilterDef;
 	MinNormalizedTireLoad = PTireLoadFilterDef.mMinNormalisedLoad;
@@ -587,7 +586,7 @@ void UWheeledVehicleMovementComponent::SetupWheels(PxVehicleWheelsSimData* PWhee
 #define ONE_TWENTYSEVENTH 0.037037f
 #define ONE_THIRD 0.33333f
 
-#if WITH_PHYSX
+#if PHYSICS_INTERFACE_PHYSX
 PX_FORCE_INLINE PxF32 smoothingFunction1(const PxF32 K)
 {
 	//Equation 20 in CarSimEd manual Appendix F.
diff --git a/Engine/Plugins/Runtime/PhysXVehicles/Source/PhysXVehicles/Private/WheeledVehicleMovementComponent4W.cpp b/Engine/Plugins/Runtime/PhysXVehicles/Source/PhysXVehicles/Private/WheeledVehicleMovementComponent4W.cpp
--- a/Engine/Plugins/Runtime/PhysXVehicles/Source/PhysXVehicles/Private/WheeledVehicleMovementComponent4W.cpp
+++ b/Engine/Plugins/Runtime/PhysXVehicles/Source/PhysXVehicles/Private/WheeledVehicleMovementComponent4W.cpp
@@ -7,7 +7,7 @@
 
 UWheeledVehicleMovementComponent4W::UWheeledVehicleMovementComponent4W(const FObjectInitializer& ObjectInitializer) : Super(ObjectInitializer)
 {
-#if WITH_PHYSX
+#if PHYSICS_INTERFACE_PHYSX
 	// grab default values from physx
 	PxVehicleDifferential4WData DefDifferentialSetup;
 	TEnumAsByte<EVehicleDifferential4W::Type> DiffType((uint8)DefDifferentialSetup.mType);
@@ -107,7 +107,7 @@ void UWheeledVehicleMovementComponent4W::PostEditChangeProperty(struct FProperty
 }
 #endif
 
-#if WITH_PHYSX
+#if PHYSICS_INTERFACE_PHYSX
 
 static void GetVehicleDifferential4WSetup(const FVehicleDifferential4WData& Setup, PxVehicleDifferential4WData& PxSetup)
 {
@@ -158,7 +158,7 @@ float FVehicleEngineData::FindPeakTorque() const
 	return PeakTorque;
 }
 
-#if WITH_PHYSX
+#if PHYSICS_INTERFACE_PHYSX
 static void GetVehicleEngineSetup(const FVehicleEngineData& Setup, PxVehicleEngineData& PxSetup)
 {
 	PxSetup.mMOI = M2ToCm2(Setup.MOI);
@@ -339,7 +339,7 @@ void UWheeledVehicleMovementComponent4W::UpdateSimulation(float DeltaTime)
 
 void UWheeledVehicleMovementComponent4W::UpdateEngineSetup(const FVehicleEngineData& NewEngineSetup)
 {
-#if WITH_PHYSX
+#if PHYSICS_INTERFACE_PHYSX
 	if (PVehicleDrive)
 	{
 		PxVehicleEngineData EngineData;
@@ -353,7 +353,7 @@ void UWheeledVehicleMovementComponent4W::UpdateEngineSetup(const FVehicleEngineD
 
 void UWheeledVehicleMovementComponent4W::UpdateDifferentialSetup(const FVehicleDifferential4WData& NewDifferentialSetup)
 {
-#if WITH_PHYSX
+#if PHYSICS_INTERFACE_PHYSX
 	if (PVehicleDrive)
 	{
 		PxVehicleDifferential4WData DifferentialData;
@@ -367,7 +367,7 @@ void UWheeledVehicleMovementComponent4W::UpdateDifferentialSetup(const FVehicleD
 
 void UWheeledVehicleMovementComponent4W::UpdateTransmissionSetup(const FVehicleTransmissionData& NewTransmissionSetup)
 {
-#if WITH_PHYSX
+#if PHYSICS_INTERFACE_PHYSX
 	if (PVehicleDrive)
 	{
 		PxVehicleGearsData GearData;
@@ -394,7 +394,7 @@ void BackwardsConvertCm2ToM2(float& val, float defaultValue)
 void UWheeledVehicleMovementComponent4W::Serialize(FArchive & Ar)
 {
 	Super::Serialize(Ar);
-#if WITH_PHYSX
+#if PHYSICS_INTERFACE_PHYSX
 	if (Ar.IsLoading() && Ar.UE4Ver() < VER_UE4_VEHICLES_UNIT_CHANGE)
 	{
 		PxVehicleEngineData DefEngineData;
diff --git a/Engine/Plugins/Runtime/RuntimePhysXCooking/Source/RuntimePhysXCooking/Private/RuntimePhysXCooking.cpp b/Engine/Plugins/Runtime/RuntimePhysXCooking/Source/RuntimePhysXCooking/Private/RuntimePhysXCooking.cpp
--- a/Engine/Plugins/Runtime/RuntimePhysXCooking/Source/RuntimePhysXCooking/Private/RuntimePhysXCooking.cpp
+++ b/Engine/Plugins/Runtime/RuntimePhysXCooking/Source/RuntimePhysXCooking/Private/RuntimePhysXCooking.cpp
@@ -1,8 +1,10 @@
 // Copyright Epic Games, Inc. All Rights Reserved.
 
-#if WITH_PHYSX
 
 #include "Modules/ModuleManager.h"
+
+#if PHYSICS_INTERFACE_PHYSX
+
 #include "PhysXCooking.h"
 
 /**
diff --git a/Engine/Source/Developer/MeshMergeUtilities/Private/MeshMergeHelpers.cpp b/Engine/Source/Developer/MeshMergeUtilities/Private/MeshMergeHelpers.cpp
--- a/Engine/Source/Developer/MeshMergeUtilities/Private/MeshMergeHelpers.cpp
+++ b/Engine/Source/Developer/MeshMergeUtilities/Private/MeshMergeHelpers.cpp
@@ -1001,9 +1001,10 @@ void FMeshMergeHelpers::ExtractPhysicsGeometry(UBodySetup* InBodySetup, const FT
 	OutAggGeom.RenderInfo = nullptr;
 	for (FKConvexElem& Elem : OutAggGeom.ConvexElems)
 	{
+#if PHYSICS_INTERFACE_PHYSX
 		Elem.SetConvexMesh(nullptr);
 		Elem.SetMirroredConvexMesh(nullptr);
-#if WITH_CHAOS
+#elif WITH_CHAOS
 		Elem.ResetChaosConvexMesh();
 #endif
 	}
diff --git a/Engine/Source/Developer/TargetPlatform/Private/TargetPlatformManagerModule.cpp b/Engine/Source/Developer/TargetPlatform/Private/TargetPlatformManagerModule.cpp
--- a/Engine/Source/Developer/TargetPlatform/Private/TargetPlatformManagerModule.cpp
+++ b/Engine/Source/Developer/TargetPlatform/Private/TargetPlatformManagerModule.cpp
@@ -21,10 +21,10 @@
 #include "PlatformInfo.h"
 #include "DesktopPlatformModule.h"
 
-#if WITH_PHYSX
+#if PHYSICS_INTERFACE_PHYSX
 #include "IPhysXCooking.h"
 #include "IPhysXCookingModule.h"
-#endif // WITH_PHYSX
+#endif // PHYSICS_INTERFACE_PHYSX
 
 DEFINE_LOG_CATEGORY_STATIC(LogTargetPlatformManager, Log, All);
 
@@ -557,7 +557,7 @@ public:
 		static bool bInitialized = false;
 		static TArray<const IPhysXCooking*> Results;
 
-#if WITH_PHYSX
+#if PHYSICS_INTERFACE_PHYSX
 		if (!bInitialized || bForceCacheUpdate)
 		{
 			bInitialized = true;
@@ -584,14 +584,14 @@ public:
 				}
 			}
 		}
-#endif // WITH_PHYSX
+#endif // PHYSICS_INTERFACE_PHYSX
 
 		return Results;
 	}
 
 	virtual const IPhysXCooking* FindPhysXCooking(FName Name) override
 	{
-#if WITH_PHYSX
+#if PHYSICS_INTERFACE_PHYSX 
 		const TArray<const IPhysXCooking*>& PhysXCooking = GetPhysXCooking();
 
 		for (int32 Index = 0; Index < PhysXCooking.Num(); Index++)
@@ -608,7 +608,7 @@ public:
 				}
 			}
 		}
-#endif // WITH_PHYSX
+#endif // PHYSICS_INTERFACE_PHYSX
 
 		return nullptr;
 	}
diff --git a/Engine/Source/Editor/PhysicsAssetEditor/Private/PhysicsAssetEditor.cpp b/Engine/Source/Editor/PhysicsAssetEditor/Private/PhysicsAssetEditor.cpp
--- a/Engine/Source/Editor/PhysicsAssetEditor/Private/PhysicsAssetEditor.cpp
+++ b/Engine/Source/Editor/PhysicsAssetEditor/Private/PhysicsAssetEditor.cpp
@@ -363,11 +363,19 @@ void FPhysicsAssetEditor::PostRedo( bool bSuccess )
 		{
 			FKConvexElem& Element = Body->AggGeom.ConvexElems[ElemIdx];
 
+#if PHYSICS_INTERFACE_PHYSX
 			if (Element.GetConvexMesh() == NULL)
 			{
 				bRecreate = true;
 				break;
 			}
+#elif WITH_CHAOS
+			if (Element.GetChaosConvexMesh() == NULL)
+			{
+				bRecreate = true;
+				break;
+			}
+#endif
 		}
 
 		if (bRecreate)
diff --git a/Engine/Source/Editor/UnrealEd/Private/EditorEngine.cpp b/Engine/Source/Editor/UnrealEd/Private/EditorEngine.cpp
--- a/Engine/Source/Editor/UnrealEd/Private/EditorEngine.cpp
+++ b/Engine/Source/Editor/UnrealEd/Private/EditorEngine.cpp
@@ -4396,12 +4396,12 @@ void UEditorEngine::CleanupPhysicsSceneThatWasInitializedForSave(UWorld* World,
 
 	World->SetPhysicsScene(nullptr);
 
-#if WITH_PHYSX
+#if PHYSICS_INTERFACE_PHYSX
 	if (GPhysCommandHandler)
 	{
 		GPhysCommandHandler->Flush();
 	}
-#endif // WITH_PHYSX
+#endif // PHYSICS_INTERFACE_PHYSX
 
 	// Update components again in case it was a world without a physics scene but did have rendered components.
 	World->UpdateWorldComponents(true, true);
diff --git a/Engine/Source/Programs/UnrealBuildTool/Configuration/ModuleRules.cs b/Engine/Source/Programs/UnrealBuildTool/Configuration/ModuleRules.cs
--- a/Engine/Source/Programs/UnrealBuildTool/Configuration/ModuleRules.cs
+++ b/Engine/Source/Programs/UnrealBuildTool/Configuration/ModuleRules.cs
@@ -1124,6 +1124,10 @@ namespace UnrealBuildTool
             if (Target.bCompilePhysX == true)
 			{
 				PrivateDependencyModuleNames.Add("PhysX");
+			}
+
+			if(Target.bCompileChaos || Target.bUseChaos || Target.bCompilePhysX)
+			{
 				PublicDefinitions.Add("WITH_PHYSX=1");
 			}
 			else
diff --git a/Engine/Source/Runtime/Engine/Classes/Components/SkeletalMeshComponent.h b/Engine/Source/Runtime/Engine/Classes/Components/SkeletalMeshComponent.h
--- a/Engine/Source/Runtime/Engine/Classes/Components/SkeletalMeshComponent.h
+++ b/Engine/Source/Runtime/Engine/Classes/Components/SkeletalMeshComponent.h
@@ -50,10 +50,12 @@ typedef FOnSkelMeshTeleportedMultiCast::FDelegate FOnSkelMeshTeleported;
 
 DECLARE_DYNAMIC_MULTICAST_DELEGATE(FOnBoneTransformsFinalized);
 
+#if PHYSICS_INTERFACE_PHYSX
 namespace physx
 {
 	class PxAggregate;
 }
+#endif
 
 UENUM()
 enum class EAnimCurveType : uint8 
diff --git a/Engine/Source/Runtime/Engine/Classes/PhysicalMaterials/PhysicalMaterial.h b/Engine/Source/Runtime/Engine/Classes/PhysicalMaterials/PhysicalMaterial.h
--- a/Engine/Source/Runtime/Engine/Classes/PhysicalMaterials/PhysicalMaterial.h
+++ b/Engine/Source/Runtime/Engine/Classes/PhysicalMaterials/PhysicalMaterial.h
@@ -15,10 +15,12 @@
 
 struct FPropertyChangedEvent;
 
+#if PHYSICS_INTERFACE_PHYSX
 namespace physx
 {
 	class PxMaterial;
 }
+#endif
 
 /** DEPRECATED Pairs desired tire friction scale with tire type */
 USTRUCT()
diff --git a/Engine/Source/Runtime/Engine/Classes/PhysicsEngine/BodySetup.h b/Engine/Source/Runtime/Engine/Classes/PhysicsEngine/BodySetup.h
--- a/Engine/Source/Runtime/Engine/Classes/PhysicsEngine/BodySetup.h
+++ b/Engine/Source/Runtime/Engine/Classes/PhysicsEngine/BodySetup.h
@@ -26,6 +26,7 @@ enum class EPhysXMeshCookFlags : uint8;
 
 DECLARE_DELEGATE_OneParam(FOnAsyncPhysicsCookFinished, bool);
 
+#if PHYSICS_INTERFACE_PHYSX
 namespace physx
 {
 	class PxTriangleMesh;
@@ -39,6 +40,7 @@ namespace physx
 	class PxTriangleMesh;
 	class PxTriangleMeshGeometry;
 }
+#endif
 
 #if WITH_CHAOS
 namespace Chaos
@@ -266,7 +268,7 @@ private:
 
 public:
 
-#if WITH_PHYSX
+#if PHYSICS_INTERFACE_PHYSX
 	/** Physics triangle mesh, created from cooked data in CreatePhysicsMeshes */
 	TArray<physx::PxTriangleMesh*> TriMeshes;
 #endif
@@ -293,7 +295,7 @@ public:
 	UPROPERTY()
 	FVector BuildScale3D;
 
-#if WITH_PHYSX
+#if PHYSICS_INTERFACE_PHYSX
 	/** References the current async cook helper. Used to be able to abort a cook task */
 	FPhysXCookHelper* CurrentCookHelper;
 #endif
diff --git a/Engine/Source/Runtime/Engine/Classes/PhysicsEngine/ConvexElem.h b/Engine/Source/Runtime/Engine/Classes/PhysicsEngine/ConvexElem.h
--- a/Engine/Source/Runtime/Engine/Classes/PhysicsEngine/ConvexElem.h
+++ b/Engine/Source/Runtime/Engine/Classes/PhysicsEngine/ConvexElem.h
@@ -51,11 +51,13 @@ private:
 	UPROPERTY()
 	FTransform Transform;
 
+#if PHYSICS_INTERFACE_PHYSX
 	/** Convex mesh for this body, created from cooked data in CreatePhysicsMeshes */
 	physx::PxConvexMesh*   ConvexMesh;
 
 	/** Convex mesh for this body, flipped across X, created from cooked data in CreatePhysicsMeshes */
 	physx::PxConvexMesh*   ConvexMeshNegX;
+#endif
 
 #if WITH_CHAOS
 	TSharedPtr<Chaos::FConvex, ESPMode::ThreadSafe> ChaosConvex;
@@ -97,6 +99,7 @@ public:
 	/** Returns the volume of this element */
 	float GetVolume(const FVector& Scale) const;
 
+#if PHYSICS_INTERFACE_PHYSX
 	/** Get the PhysX convex mesh (defined in BODY space) for this element */
 	ENGINE_API physx::PxConvexMesh* GetConvexMesh() const;
 
@@ -108,6 +111,7 @@ public:
 
 	/** Set the PhysX convex mesh to use for this element */
 	ENGINE_API void SetMirroredConvexMesh(physx::PxConvexMesh* InMesh);
+#endif
 
 #if WITH_CHAOS
 	ENGINE_API const auto& GetChaosConvexMesh() const
diff --git a/Engine/Source/Runtime/Engine/Classes/PhysicsEngine/PhysicalAnimationComponent.h b/Engine/Source/Runtime/Engine/Classes/PhysicsEngine/PhysicalAnimationComponent.h
--- a/Engine/Source/Runtime/Engine/Classes/PhysicsEngine/PhysicalAnimationComponent.h
+++ b/Engine/Source/Runtime/Engine/Classes/PhysicsEngine/PhysicalAnimationComponent.h
@@ -12,10 +12,12 @@
 class FPrimitiveDrawInterface;
 class USkeletalMeshComponent;
 
+#if PHYSICS_INTERFACE_PHYSX
 namespace physx
 {
 	class PxRigidDynamic;
 }
+#endif
 
 /** Stores info on the type of motor that will be used for a given bone */
 USTRUCT(BlueprintType)
@@ -127,7 +129,7 @@ private:
 	struct FPhysicalAnimationInstanceData
 	{
 		struct FConstraintInstance* ConstraintInstance;
-#if WITH_PHYSX
+#if PHYSICS_INTERFACE_PHYSX
 		physx::PxRigidDynamic* TargetActor; // #PHYS2 Should prob change to FPhysicsActorReference?
 #endif
 	};
diff --git a/Engine/Source/Runtime/Engine/Private/Collision/CollisionDebugDrawing.cpp b/Engine/Source/Runtime/Engine/Private/Collision/CollisionDebugDrawing.cpp
--- a/Engine/Source/Runtime/Engine/Private/Collision/CollisionDebugDrawing.cpp
+++ b/Engine/Source/Runtime/Engine/Private/Collision/CollisionDebugDrawing.cpp
@@ -201,7 +201,7 @@ ENGINE_API void DrawCapsuleOverlap(const UWorld* InWorld,const FVector& Pos, con
 	}
 }
 
-#if WITH_PHYSX
+#if PHYSICS_INTERFACE_PHYSX
 
 
 void DrawGeomOverlaps(const UWorld* InWorld, const PxGeometry& PGeom, const PxTransform& PGeomPose, TArray<struct FOverlapResult>& Overlaps, float Lifetime)
diff --git a/Engine/Source/Runtime/Engine/Private/Collision/CollisionDebugDrawing.h b/Engine/Source/Runtime/Engine/Private/Collision/CollisionDebugDrawing.h
--- a/Engine/Source/Runtime/Engine/Private/Collision/CollisionDebugDrawing.h
+++ b/Engine/Source/Runtime/Engine/Private/Collision/CollisionDebugDrawing.h
@@ -11,7 +11,7 @@
 #include "PhysXIncludes.h"
 #include "Chaos/ImplicitObject.h"
 
-#if WITH_PHYSX
+#if PHYSICS_INTERFACE_PHYSX
 /** Draw PhysX geom with overlaps */
 void DrawGeomOverlaps(const UWorld* InWorld, const PxGeometry& PGeom, const PxTransform& PGeomPose, TArray<struct FOverlapResult>& Overlaps, float Lifetime);
 /** Draw PhysX geom being swept with hits */
diff --git a/Engine/Source/Runtime/Engine/Private/Collision/KAggregateGeom.cpp b/Engine/Source/Runtime/Engine/Private/Collision/KAggregateGeom.cpp
--- a/Engine/Source/Runtime/Engine/Private/Collision/KAggregateGeom.cpp
+++ b/Engine/Source/Runtime/Engine/Private/Collision/KAggregateGeom.cpp
@@ -416,7 +416,7 @@ FBox FKConvexElem::CalcAABB(const FTransform& BoneTM, const FVector& Scale3D) co
 
 void FKConvexElem::GetPlanes(TArray<FPlane>& Planes) const
 {
-#if WITH_PHYSX
+#if PHYSICS_INTERFACE_PHYSX
 	if (ConvexMesh != nullptr)
 	{
 		Planes.Empty();
@@ -435,7 +435,9 @@ void FKConvexElem::GetPlanes(TArray<FPlane>& Planes) const
 			Planes.Add(Plane);
 		}
 	}
-#endif // WITH_PHYSX
+#else
+	CHAOS_ENSURE(false); // TODO Implement in Chaos
+#endif
 }
 
 ///////////////////////////////////////
@@ -774,8 +776,10 @@ FArchive& operator<<(FArchive& Ar,FKConvexElem& Elem)
 		// Initialize the TArray members
 		FMemory::Memzero(&Elem.VertexData, sizeof(Elem.VertexData));
 		FMemory::Memzero(&Elem.ElemBox, sizeof(Elem.ElemBox));
+#if PHYSICS_INTERFACE_PHYSX
 		Elem.ConvexMesh = NULL;
 		Elem.ConvexMeshNegX = NULL;
+#endif
 #if WITH_CHAOS
 		Elem.ChaosConvex.Reset();
 #endif
diff --git a/Engine/Source/Runtime/Engine/Private/Collision/PhysicsFiltering.cpp b/Engine/Source/Runtime/Engine/Private/Collision/PhysicsFiltering.cpp
--- a/Engine/Source/Runtime/Engine/Private/Collision/PhysicsFiltering.cpp
+++ b/Engine/Source/Runtime/Engine/Private/Collision/PhysicsFiltering.cpp
@@ -28,6 +28,6 @@ FPhysicsFilterBuilder::FPhysicsFilterBuilder(TEnumAsByte<enum ECollisionChannel>
 	Word3 = CreateChannelAndFilter(InObjectType, MaskFilter);
 }
 
-#if WITH_PHYSX
+#if PHYSICS_INTERFACE_PHYSX
 PxSimulationFilterShader GSimulationFilterShader = nullptr;
 #endif // WITH_PHYSX
diff --git a/Engine/Source/Runtime/Engine/Private/Collision/SceneQueryLowLevel.cpp b/Engine/Source/Runtime/Engine/Private/Collision/SceneQueryLowLevel.cpp
--- a/Engine/Source/Runtime/Engine/Private/Collision/SceneQueryLowLevel.cpp
+++ b/Engine/Source/Runtime/Engine/Private/Collision/SceneQueryLowLevel.cpp
@@ -1,17 +1,16 @@
 // Copyright Epic Games, Inc. All Rights Reserved.
 
-#if WITH_PHYSX
+#if PHYSICS_INTERFACE_PHYSX
 #include "PhysXPublic.h"
+#include "PhysXInterfaceWrapper.h"
 #endif
+
 #include "Physics/PhysicsInterfaceDeclares.h"
 #include "Physics/PhysicsInterfaceCore.h"
 #include "PhysicsInterfaceDeclaresCore.h"
 
 #include "PhysicsEngine/CollisionQueryFilterCallback.h"
 #include "PhysicsCore.h"
-#if PHYSICS_INTERFACE_PHYSX
-#include "PhysXInterfaceWrapper.h"
-#endif
 
 #include "PhysTestSerializer.h"
 
@@ -47,7 +46,7 @@ void FinalizeCapture(FPhysTestSerializer& Serializer)
 	{
 		Serializer.Serialize(TEXT("SQCapture"));
 	}
-#if WITH_PHYSX
+#if 0
 	if (ReplaySQs)
 	{
 		const bool bReplaySuccess = SQComparisonHelper(Serializer);
diff --git a/Engine/Source/Runtime/Engine/Private/Components/SplineMeshComponent.cpp b/Engine/Source/Runtime/Engine/Private/Components/SplineMeshComponent.cpp
--- a/Engine/Source/Runtime/Engine/Private/Components/SplineMeshComponent.cpp
+++ b/Engine/Source/Runtime/Engine/Private/Components/SplineMeshComponent.cpp
@@ -906,7 +906,7 @@ void USplineMeshComponent::OnCreatePhysicsState()
 
 UBodySetup* USplineMeshComponent::GetBodySetup()
 {
-#if WITH_PHYSX
+#if PHYSICS_INTERFACE_PHYSX
 	// Don't return a body setup that has no collision, it means we are interactively moving the spline and don't want to build collision.
 	// Instead we explicitly build collision with USplineMeshComponent::RecreateCollision()
 	if (BodySetup != NULL && (BodySetup->TriMeshes.Num() || BodySetup->AggGeom.GetElementCount() > 0))
diff --git a/Engine/Source/Runtime/Engine/Private/LevelTick.cpp b/Engine/Source/Runtime/Engine/Private/LevelTick.cpp
--- a/Engine/Source/Runtime/Engine/Private/LevelTick.cpp
+++ b/Engine/Source/Runtime/Engine/Private/LevelTick.cpp
@@ -1693,7 +1693,7 @@ void UWorld::Tick( ELevelTick TickType, float DeltaSeconds )
 
 		FWorldDelegates::OnWorldPostActorTick.Broadcast(this, TickType, DeltaSeconds);
 
-#if WITH_PHYSX
+#if PHYSICS_INTERFACE_PHYSX
 		if ( PhysicsScene != nullptr )
 		{
 			GPhysCommandHandler->Flush();
diff --git a/Engine/Source/Runtime/Engine/Private/PhysicsEngine/BodyInstance.cpp b/Engine/Source/Runtime/Engine/Private/PhysicsEngine/BodyInstance.cpp
--- a/Engine/Source/Runtime/Engine/Private/PhysicsEngine/BodyInstance.cpp
+++ b/Engine/Source/Runtime/Engine/Private/PhysicsEngine/BodyInstance.cpp
@@ -43,7 +43,7 @@
 #include "Chaos/ParticleHandle.h"
 #endif // WITH_CHAOS
 
-#if WITH_PHYSX
+#if PHYSICS_INTERFACE_PHYSX
 	#include "PhysXPublic.h"
 	#include "Physics/PhysicsFiltering.h"
 	#include "PhysicsEngine/PhysXSupport.h"
@@ -453,7 +453,7 @@ void FBodyInstance::UpdateTriMeshVertices(const TArray<FVector> & NewPositions)
 #if WITH_CHAOS || WITH_IMMEDIATE_PHYSX
 	ensure(false);
 #else
-#if WITH_PHYSX
+#if PHYSICS_INTERFACE_PHYSX
 	if (BodySetup.IsValid())
 	{
 		FPhysicsCommand::ExecuteWrite(ActorHandle, [&](const FPhysicsActorHandle& Actor)
@@ -2006,7 +2006,7 @@ bool FBodyInstance::UpdateBodyScale(const FVector& InScale3D, bool bForceUpdate)
 		}
 	}
 
-#elif WITH_PHYSX
+#elif PHYSICS_INTERFACE_PHYSX
 	//Get all shapes
 	EScaleMode::Type ScaleMode = EScaleMode::Free;
 
@@ -2940,7 +2940,7 @@ Chaos::TMassProperties<float, 3> ComputeMassProperties(const FBodyInstance* Owni
 	}
 
 	Chaos::TMassProperties<float, 3> MassProps;
-	FPhysicsInterface::CalculateTMassPropertiesFromShapeCollection(MassProps, Shapes, DensityKGPerCubicUU);
+	FPhysicsInterface::CalculateMassPropertiesFromShapeCollection(MassProps, Shapes, DensityKGPerCubicUU);
 
 	float OldMass = MassProps.Mass;
 	float NewMass = 0.f;
@@ -2968,7 +2968,7 @@ Chaos::TMassProperties<float, 3> ComputeMassProperties(const FBodyInstance* Owni
 
 	MassProps.CenterOfMass += MassModifierTransform.TransformVector(OwningBodyInstance->COMNudge);
 
-	// Scale the inertia tensor by the owning body instance's InertiaTensorScale
+	// Scale the inertia tensor by the owning body instance's InertiaTensorScale (see PxMassProperties::scaleInertia)
 	FVector diagonal(MassProps.InertiaTensor.M[0][0] , MassProps.InertiaTensor.M[1][1], MassProps.InertiaTensor.M[2][2]);
 
 	FVector xyz2 = FVector(FVector::DotProduct(diagonal, FVector(0.5f))) - diagonal; // original x^2, y^2, z^2
@@ -2988,7 +2988,7 @@ Chaos::TMassProperties<float, 3> ComputeMassProperties(const FBodyInstance* Owni
 
 	return MassProps;
 }
-#elif WITH_PHYSX
+#elif PHYSICS_INTERFACE_PHYSX
 /** Computes and adds the mass properties (inertia, com, etc...) based on the mass settings of the body instance. */
 PxMassProperties ComputeMassProperties(const FBodyInstance* OwningBodyInstance, TArray<FPhysicsShapeHandle> Shapes, const FTransform& MassModifierTransform)
 {
@@ -3037,6 +3037,7 @@ void FBodyInstance::UpdateMassProperties()
 {
 	UPhysicalMaterial* PhysMat = GetSimplePhysicalMaterial();
 
+#if WITH_PHYSX
 	if (FPhysicsInterface::IsValid(ActorHandle) && FPhysicsInterface::IsRigidBody(ActorHandle))
 	{
 		FPhysicsCommand::ExecuteWrite(ActorHandle, [&](FPhysicsActorHandle& Actor)
@@ -3066,7 +3067,7 @@ void FBodyInstance::UpdateMassProperties()
 
 #if WITH_CHAOS
 				Chaos::TMassProperties<float, 3> TotalMassProperties;
-#elif WITH_PHYSX
+#elif PHYSICS_INTERFACE_PHYSX
 				PxMassProperties TotalMassProperties;
 #endif
 				if (ShapeToBodiesMap.IsValid() && ShapeToBodiesMap->Num() > 0)
@@ -3121,7 +3122,7 @@ void FBodyInstance::UpdateMassProperties()
 					}
 
 					TotalMassProperties = Chaos::Combine(SubMassProperties);
-#elif WITH_PHYSX
+#elif PHYSICS_INTERFACE_PHYSX
 					TArray<PxMassProperties> SubMassProperties;
 					TArray<PxTransform> MassTMs;
 					for (auto BodyShapesItr : BodyToShapes)
@@ -3180,6 +3181,7 @@ void FBodyInstance::UpdateMassProperties()
 	{
 		BodyInstanceDelegates->OnRecalculatedMassProperties.Broadcast(this);
 	}
+#endif
 }
 
 void FBodyInstance::UpdateDebugRendering()
diff --git a/Engine/Source/Runtime/Engine/Private/PhysicsEngine/BodySetup.cpp b/Engine/Source/Runtime/Engine/Private/PhysicsEngine/BodySetup.cpp
--- a/Engine/Source/Runtime/Engine/Private/PhysicsEngine/BodySetup.cpp
+++ b/Engine/Source/Runtime/Engine/Private/PhysicsEngine/BodySetup.cpp
@@ -26,7 +26,7 @@
 
 #include "PhysXCookHelper.h"
 
-#if WITH_PHYSX
+#if PHYSICS_INTERFACE_PHYSX
 	#include "PhysXPublic.h"
 	#include "PhysicsEngine/PhysXSupport.h"
 #endif // WITH_PHYSX
@@ -128,12 +128,15 @@ DEFINE_STAT(STAT_PhysXCooking);
 
 bool IsRuntimeCookingEnabled()
 {
+#if PHYSICS_INTERFACE_PHYSX
 	return FModuleManager::LoadModulePtr<IPhysXCookingModule>("RuntimePhysXCooking") != nullptr;
+#else
+	return false;
+#endif
 }
 #endif //WITH_PHYSX
 
-
-#if WITH_PHYSX
+#if PHYSICS_INTERFACE_PHYSX
 	// Quaternion that converts Sphyls from UE space to PhysX space (negate Y, swap X & Z)
 	// This is equivalent to a 180 degree rotation around the normalized (1, 0, 1) axis
 	const physx::PxQuat U2PSphylBasis( PI, PxVec3( 1.0f / FMath::Sqrt( 2.0f ), 0.0f, 1.0f / FMath::Sqrt( 2.0f ) ) );
@@ -188,7 +191,7 @@ UBodySetup::UBodySetup(const FObjectInitializer& ObjectInitializer)
 	SetFlags(RF_Transactional);
 	bSharedCookedData = false;
 	CookedFormatDataOverride = nullptr;
-#if WITH_PHYSX
+#if PHYSICS_INTERFACE_PHYSX
 	CurrentCookHelper = nullptr;
 #endif
 }
@@ -201,9 +204,10 @@ void UBodySetup::CopyBodyPropertiesFrom(const UBodySetup* FromSetup)
 	for (int32 i = 0; i < AggGeom.ConvexElems.Num(); i++)
 	{
 		FKConvexElem& ConvexElem = AggGeom.ConvexElems[i];
+#if PHYSICS_INTERFACE_PHYSX
 		ConvexElem.SetConvexMesh(nullptr);
 		ConvexElem.SetMirroredConvexMesh(nullptr);
-#if WITH_CHAOS
+#elif WITH_CHAOS
 		ConvexElem.ResetChaosConvexMesh();
 #endif
 	}
@@ -230,9 +234,10 @@ void UBodySetup::AddCollisionFrom(const FKAggregateGeom& FromAggGeom)
 	for (int32 i = FirstNewConvexIdx; i < AggGeom.ConvexElems.Num(); i++)
 	{
 		FKConvexElem& ConvexElem = AggGeom.ConvexElems[i];
+#if PHYSICS_INTERFACE_PHYSX
 		ConvexElem.SetConvexMesh(nullptr);
 		ConvexElem.SetMirroredConvexMesh(nullptr);
-#if WITH_CHAOS
+#elif WITH_CHAOS
 		ConvexElem.ResetChaosConvexMesh();
 #endif
 	}
@@ -636,7 +641,7 @@ void UBodySetup::CreatePhysicsMeshesAsync(FOnAsyncPhysicsCookFinished OnAsyncPhy
 
 void UBodySetup::AbortPhysicsMeshAsyncCreation()
 {
-#if WITH_PHYSX
+#if PHYSICS_INTERFACE_PHYSX
 	if (CurrentCookHelper)
 	{
 		CurrentCookHelper->Abort();
@@ -855,7 +860,7 @@ void UBodySetup::AddShapesToRigidActor_AssumesLocked(
 	AddParams.LocalTransform = RelativeTM;
 	AddParams.WorldTransform = OwningInstance->GetUnrealWorldTransform();
 	AddParams.Geometry = &AggGeom;
-#if WITH_PHYSX
+#if PHYSICS_INTERFACE_PHYSX
 	AddParams.TriMeshes = TArrayView<PxTriangleMesh*>(TriMeshes);
 #endif
 
@@ -1152,7 +1157,7 @@ void UBodySetup::PostLoad()
 void UBodySetup::UpdateTriMeshVertices(const TArray<FVector> & NewPositions)
 {
 	SCOPE_CYCLE_COUNTER(STAT_UpdateTriMeshVertices);
-#if WITH_PHYSX
+#if PHYSICS_INTERFACE_PHYSX
 	if (TriMeshes.Num())
 	{
 		check(TriMeshes[0] != nullptr);
@@ -1504,7 +1509,7 @@ void UBodySetup::GetResourceSizeEx(FResourceSizeEx& CumulativeResourceSize)
 {
 	Super::GetResourceSizeEx(CumulativeResourceSize);
 
-#if WITH_PHYSX
+#if PHYSICS_INTERFACE_PHYSX
 	// Count PhysX trimesh mem usage
 	for(PxTriangleMesh* TriMesh : TriMeshes)
 	{
@@ -1632,13 +1637,17 @@ FKConvexElem::FKConvexElem()
 	: FKShapeElem(EAggCollisionShape::Convex)
 	, ElemBox(ForceInit)
 	, Transform(FTransform::Identity)
+#if PHYSICS_INTERFACE_PHYSX
 	, ConvexMesh(NULL)
 	, ConvexMeshNegX(NULL)
+#endif
 {}
 
 FKConvexElem::FKConvexElem(const FKConvexElem& Other)
+#if PHYSICS_INTERFACE_PHYSX
 	: ConvexMesh(nullptr)
 	, ConvexMeshNegX(nullptr)
+#endif
 {
 	CloneElem(Other);
 }
@@ -1650,11 +1659,12 @@ FKConvexElem::~FKConvexElem()
 
 const FKConvexElem& FKConvexElem::operator=(const FKConvexElem& Other)
 {
+#if PHYSICS_INTERFACE_PHYSX
 	ensureMsgf(ConvexMesh == nullptr, TEXT("We are leaking memory. Why are we calling the assignment operator on an element that has already allocated resources?"));
 	ensureMsgf(ConvexMeshNegX == nullptr, TEXT("We are leaking memory. Why are we calling the assignment operator on an element that has already allocated resources?"));
 	ConvexMesh = nullptr;
 	ConvexMeshNegX = nullptr;
-#if WITH_CHAOS
+#elif WITH_CHAOS
 	ensureMsgf(!ChaosConvex, TEXT("We are leaking memory. Why are we calling the assignment operator on an element that has already allocated resources?"));
 	ResetChaosConvexMesh();
 #endif
@@ -1686,7 +1696,7 @@ float SignedVolumeOfTriangle(const FVector& p1, const FVector& p2, const FVector
 {
 	return FVector::DotProduct(p1, FVector::CrossProduct(p2, p3)) / 6.0f;
 }
-
+#if PHYSICS_INTERFACE_PHYSX
 physx::PxConvexMesh* FKConvexElem::GetConvexMesh() const
 {
 	return ConvexMesh;
@@ -1706,12 +1716,13 @@ void FKConvexElem::SetMirroredConvexMesh(physx::PxConvexMesh* InMesh)
 {
 	ConvexMeshNegX = InMesh;
 }
+#endif
 
 float FKConvexElem::GetVolume(const FVector& Scale) const
 {
 	float Volume = 0.0f;
 
-#if WITH_PHYSX
+#if PHYSICS_INTERFACE_PHYSX
 	if (ConvexMesh != NULL)
 	{
 		// Preparation for convex mesh scaling implemented in another changelist
@@ -1742,7 +1753,10 @@ float FKConvexElem::GetVolume(const FVector& Scale) const
 			}
 		}
 	}
-#endif // WITH_PHYSX
+#elif WITH_CHAOS
+	//TODO Support ChaosConvex.
+	CHAOS_ENSURE(false);
+#endif
 
 	return Volume;
 }
diff --git a/Engine/Source/Runtime/Engine/Private/PhysicsEngine/CollisionQueryFilterCallback.cpp b/Engine/Source/Runtime/Engine/Private/PhysicsEngine/CollisionQueryFilterCallback.cpp
--- a/Engine/Source/Runtime/Engine/Private/PhysicsEngine/CollisionQueryFilterCallback.cpp
+++ b/Engine/Source/Runtime/Engine/Private/PhysicsEngine/CollisionQueryFilterCallback.cpp
@@ -8,7 +8,7 @@
 #include "Components/PrimitiveComponent.h"
 #include "Collision.h"
 
-#if WITH_PHYSX
+#if PHYSICS_INTERFACE_PHYSX
 #include "PhysXInterfaceWrapper.h"
 #endif
 
@@ -86,7 +86,7 @@ ECollisionQueryHitType FCollisionQueryFilterCallback::CalcQueryHitType(const FCo
 	return ECollisionQueryHitType::None;
 }
 
-#if WITH_PHYSX
+#if PHYSICS_INTERFACE_PHYSX
 ECollisionQueryHitType FCollisionQueryFilterCallback::PreFilterImp(const FCollisionFilterData& FilterData, const physx::PxShape& Shape, const physx::PxActor& Actor)
 {
 	//SCOPE_CYCLE_COUNTER(STAT_Collision_PreFilter);
@@ -251,7 +251,7 @@ ECollisionQueryHitType FCollisionQueryFilterCallback::PostFilterImp(const FColli
 	}
 }
 
-#if WITH_PHYSX
+#if PHYSICS_INTERFACE_PHYSX
 ECollisionQueryHitType FCollisionQueryFilterCallback::PostFilterImp(const FCollisionFilterData& FilterData, const physx::PxQueryHit& Hit)
 {
 	// Unused in non-sweeps
@@ -281,7 +281,7 @@ ECollisionQueryHitType FCollisionQueryFilterCallback::PostFilterImp(const FColli
 	return PostFilterImp(FilterData, bIsOverlap);
 }
 
-#if WITH_PHYSX
+#if PHYSICS_INTERFACE_PHYSX
 
 PxQueryHitType::Enum FCollisionQueryFilterCallback::preFilter(const PxFilterData& filterData, const PxShape* shape, const PxRigidActor* actor, PxHitFlags& queryFlags)
 {
diff --git a/Engine/Source/Runtime/Engine/Private/PhysicsEngine/CollisionQueryFilterCallback.h b/Engine/Source/Runtime/Engine/Private/PhysicsEngine/CollisionQueryFilterCallback.h
--- a/Engine/Source/Runtime/Engine/Private/PhysicsEngine/CollisionQueryFilterCallback.h
+++ b/Engine/Source/Runtime/Engine/Private/PhysicsEngine/CollisionQueryFilterCallback.h
@@ -64,7 +64,7 @@ public:
 
 	ECollisionQueryHitType PostFilterImp(const FCollisionFilterData& FilterData, bool bIsOverlap);
 
-#if WITH_PHYSX
+#if PHYSICS_INTERFACE_PHYSX
 	ECollisionQueryHitType PostFilterImp(const FCollisionFilterData& FilterData, const physx::PxQueryHit& Hit);
 	ECollisionQueryHitType PreFilterImp(const FCollisionFilterData& FilterData, const physx::PxShape& Shape, const physx::PxActor& Actor);
 #endif
@@ -81,7 +81,7 @@ public:
 		return PreFilterImp(FilterData, Shape, Actor);
 	}
 
-#if WITH_PHYSX
+#if PHYSICS_INTERFACE_PHYSX
 	virtual ECollisionQueryHitType PostFilter(const FCollisionFilterData& FilterData, const physx::PxQueryHit& Hit) override { return PostFilterImp(FilterData, Hit); }
 	virtual ECollisionQueryHitType PreFilter(const FCollisionFilterData& FilterData, const physx::PxShape& Shape, physx::PxRigidActor& Actor) override { return PreFilterImp(FilterData, Shape, Actor); }
 	virtual PxQueryHitType::Enum preFilter(const PxFilterData& filterData, const PxShape* shape, const PxRigidActor* actor, PxHitFlags& queryFlags) override;
diff --git a/Engine/Source/Runtime/Engine/Private/PhysicsEngine/ConstraintTypes.cpp b/Engine/Source/Runtime/Engine/Private/PhysicsEngine/ConstraintTypes.cpp
--- a/Engine/Source/Runtime/Engine/Private/PhysicsEngine/ConstraintTypes.cpp
+++ b/Engine/Source/Runtime/Engine/Private/PhysicsEngine/ConstraintTypes.cpp
@@ -11,7 +11,7 @@ extern TAutoConsoleVariable<float> CVarConstraintLinearStiffnessScale;
 extern TAutoConsoleVariable<float> CVarConstraintAngularDampingScale;
 extern TAutoConsoleVariable<float> CVarConstraintAngularStiffnessScale;
 
-#if WITH_PHYSX
+#if PHYSICS_INTERFACE_PHYSX
 
 enum class ESoftLimitTypeHelper
 {
diff --git a/Engine/Source/Runtime/Engine/Private/PhysicsEngine/Experimental/ChaosInterfaceUtils.cpp b/Engine/Source/Runtime/Engine/Private/PhysicsEngine/Experimental/ChaosInterfaceUtils.cpp
--- a/Engine/Source/Runtime/Engine/Private/PhysicsEngine/Experimental/ChaosInterfaceUtils.cpp
+++ b/Engine/Source/Runtime/Engine/Private/PhysicsEngine/Experimental/ChaosInterfaceUtils.cpp
@@ -23,7 +23,7 @@
 #include "PhysicsEngine/AggregateGeom.h"
 #include "PhysicsEngine/PhysicsSettings.h"
 
-#if WITH_PHYSX
+#if PHYSICS_INTERFACE_PHYSX
 #include "PhysXIncludes.h"
 #endif
 
@@ -37,7 +37,7 @@ namespace ChaosInterface
 		check(false);
 	}
 
-#if WITH_PHYSX
+#if PHYSICS_INTERFACE_PHYSX
 
 	template<>
 	TArray<Chaos::TVector<int32, 3>> GetMeshElements(const physx::PxConvexMesh* PhysXMesh)
diff --git a/Engine/Source/Runtime/Engine/Private/PhysicsEngine/Experimental/PhysInterface_Chaos.cpp b/Engine/Source/Runtime/Engine/Private/PhysicsEngine/Experimental/PhysInterface_Chaos.cpp
--- a/Engine/Source/Runtime/Engine/Private/PhysicsEngine/Experimental/PhysInterface_Chaos.cpp
+++ b/Engine/Source/Runtime/Engine/Private/PhysicsEngine/Experimental/PhysInterface_Chaos.cpp
@@ -41,7 +41,7 @@
 #include "PhysicalMaterials/PhysicalMaterialMask.h"
 #include "PhysicalMaterials/PhysicalMaterial.h"
 
-#if WITH_PHYSX
+#if PHYSICS_INTERFACE_PHYSX
 #include "geometry/PxConvexMesh.h"
 #include "geometry/PxTriangleMesh.h"
 #include "foundation/PxVec3.h"
@@ -1789,7 +1789,7 @@ bool CalculateMassPropertiesOfImplicitType(
 	return false;
 }
 
-void FPhysInterface_Chaos::CalculateTMassPropertiesFromShapeCollection(Chaos::TMassProperties<float, 3>& OutProperties, const TArray<FPhysicsShapeHandle>& InShapes, float InDensityKGPerCM)
+void FPhysInterface_Chaos::CalculateMassPropertiesFromShapeCollection(Chaos::TMassProperties<float, 3>& OutProperties, const TArray<FPhysicsShapeHandle>& InShapes, float InDensityKGPerCM)
 {
 	float TotalMass = 0.f;
 	Chaos::FVec3 TotalCenterOfMass(0.f);
@@ -1834,56 +1834,6 @@ void FPhysInterface_Chaos::CalculateTMassPropertiesFromShapeCollection(Chaos::TM
 	OutProperties.CenterOfMass = TotalCenterOfMass;
 }
 
-#if WITH_PHYSX
-
-void FPhysInterface_Chaos::CalculateMassPropertiesFromShapeCollection(physx::PxMassProperties& OutProperties, const TArray<FPhysicsShapeHandle>& InShapes, float InDensityKGPerCM)
-{
-	float TotalMass = 0.f;
-	Chaos::FVec3 TotalCenterOfMass(0.f);
-	TArray< Chaos::TMassProperties<float,3> > MassPropertiesList;
-	for (const FPhysicsShapeHandle& ShapeHandle : InShapes)
-	{
-		if (const Chaos::TPerShapeData<float, 3>* Shape = ShapeHandle.Shape)
-		{
-			if (const Chaos::FImplicitObject * ImplicitObject = Shape->Geometry.Get())
-			{
-				FTransform WorldTransform(ShapeHandle.ActorRef->R(), ShapeHandle.ActorRef->X());
-				Chaos::TMassProperties<float,3> MassProperties;
-				if (CalculateMassPropertiesOfImplicitType(MassProperties, WorldTransform, ImplicitObject, InDensityKGPerCM))
-				{
-					MassPropertiesList.Add(MassProperties);
-					TotalMass += MassProperties.Mass;
-					TotalCenterOfMass += MassProperties.CenterOfMass * MassProperties.Mass;
-				}
-			}
-		}
-	}
-
-	if (TotalMass > 0.f)
-	{
-		TotalCenterOfMass /= TotalMass;
-	}
-
-	Chaos::PMatrix<float, 3, 3> Tensor;
-	if (MassPropertiesList.Num())
-	{
-		Tensor = Chaos::CombineWorldSpace<float, 3>(MassPropertiesList,  InDensityKGPerCM).InertiaTensor;
-	}
-	else 
-	{
-		// @todo : Add support for all types, but for now just hard code a unit sphere tensor {r:50cm} if the type was not processed
-		Tensor = Chaos::PMatrix<float, 3, 3>(5.24e5, 5.24e5, 5.24e5);
-		TotalMass = 523.f;
-	}
-
-	float Mat[] = { Tensor.M[0][0],Tensor.M[0][1],Tensor.M[0][2],Tensor.M[1][0],Tensor.M[1][1],Tensor.M[1][2],Tensor.M[2][0],Tensor.M[2][1],Tensor.M[2][2] };
-	OutProperties.inertiaTensor = PxMat33(Mat);
-	OutProperties.mass = TotalMass;
-	OutProperties.centerOfMass = U2PVector(TotalCenterOfMass);
-}
-
-#endif
-
 bool FPhysInterface_Chaos::LineTrace_Geom(FHitResult& OutHit, const FBodyInstance* InInstance, const FVector& WorldStart, const FVector& WorldEnd, bool bTraceComplex, bool bExtractPhysMaterial)
 {
 	// Need an instance to trace against
diff --git a/Engine/Source/Runtime/Engine/Private/PhysicsEngine/ImmediatePhysics/ImmediatePhysicsChaos/ImmediatePhysicsActorHandle_Chaos.cpp b/Engine/Source/Runtime/Engine/Private/PhysicsEngine/ImmediatePhysics/ImmediatePhysicsChaos/ImmediatePhysicsActorHandle_Chaos.cpp
--- a/Engine/Source/Runtime/Engine/Private/PhysicsEngine/ImmediatePhysics/ImmediatePhysicsChaos/ImmediatePhysicsActorHandle_Chaos.cpp
+++ b/Engine/Source/Runtime/Engine/Private/PhysicsEngine/ImmediatePhysics/ImmediatePhysicsChaos/ImmediatePhysicsActorHandle_Chaos.cpp
@@ -229,7 +229,7 @@ namespace ImmediatePhysics_Chaos
 #endif
 		AddParams.WorldTransform = BodyInstance->GetUnrealWorldTransform();
 		AddParams.Geometry = &BodySetup->AggGeom;
-#if WITH_PHYSX
+#if PHYSICS_INTERFACE_PHYSX
 		AddParams.TriMeshes = TArrayView<PxTriangleMesh*>(BodySetup->TriMeshes);
 #endif
 #if WITH_CHAOS
diff --git a/Engine/Source/Runtime/Engine/Private/PhysicsEngine/PhysCommandHandler.cpp b/Engine/Source/Runtime/Engine/Private/PhysicsEngine/PhysCommandHandler.cpp
--- a/Engine/Source/Runtime/Engine/Private/PhysicsEngine/PhysCommandHandler.cpp
+++ b/Engine/Source/Runtime/Engine/Private/PhysicsEngine/PhysCommandHandler.cpp
@@ -6,7 +6,7 @@
 #include "PhysicsPublic.h"
 #include "Physics/PhysicsInterfaceCore.h"
 
-#if WITH_PHYSX
+#if PHYSICS_INTERFACE_PHYSX
 	#include "PhysicsEngine/PhysXSupport.h"
 #endif
 
@@ -48,8 +48,10 @@ void FPhysCommandHandler::ExecuteCommands()
 #if WITH_PHYSX
 		case PhysCommand::ReleasePScene:
 		{
+#if PHYSICS_INTERFACE_PHYSX
 			physx::PxScene * PScene = Command.Pointer.PScene;
 			PScene->release();
+#endif
 			break;
 		}
 		case PhysCommand::DeleteSimEventCallback:
@@ -102,15 +104,19 @@ void FPhysCommandHandler::ExecuteCommands()
 
 		case PhysCommand::DeleteCPUDispatcher:
 		{
+#if PHYSICS_INTERFACE_PHYSX
 			physx::PxCpuDispatcher * CPUDispatcher = Command.Pointer.CPUDispatcher;
 			delete CPUDispatcher;
+#endif
 			break;
 		}
 
 		case PhysCommand::DeleteMbpBroadphaseCallback:
 		{
+#if PHYSICS_INTERFACE_PHYSX
 			FPhysXMbpBroadphaseCallback* Callback = Command.Pointer.MbpCallback;
 			delete Callback;
+#endif
 			break;
 		}
 #endif
@@ -129,7 +135,7 @@ void FPhysCommandHandler::EnqueueCommand(const FPhysPendingCommand& Command)
 }
 
 
-#if WITH_PHYSX
+#if PHYSICS_INTERFACE_PHYSX
 void FPhysCommandHandler::DeferredRelease(physx::PxScene* PScene)
 {
 	check(PScene);
diff --git a/Engine/Source/Runtime/Engine/Private/PhysicsEngine/PhysDerivedData.cpp b/Engine/Source/Runtime/Engine/Private/PhysicsEngine/PhysDerivedData.cpp
--- a/Engine/Source/Runtime/Engine/Private/PhysicsEngine/PhysDerivedData.cpp
+++ b/Engine/Source/Runtime/Engine/Private/PhysicsEngine/PhysDerivedData.cpp
@@ -9,7 +9,7 @@
 #include "PhysicsEngine/BodySetup.h"
 #include "Interfaces/ITargetPlatformManagerModule.h"
 
-#if WITH_PHYSX && WITH_EDITOR
+#if PHYSICS_INTERFACE_PHYSX && WITH_EDITOR
 
 #include "IPhysXCookingModule.h"
 
diff --git a/Engine/Source/Runtime/Engine/Private/PhysicsEngine/PhysDerivedData.h b/Engine/Source/Runtime/Engine/Private/PhysicsEngine/PhysDerivedData.h
--- a/Engine/Source/Runtime/Engine/Private/PhysicsEngine/PhysDerivedData.h
+++ b/Engine/Source/Runtime/Engine/Private/PhysicsEngine/PhysDerivedData.h
@@ -9,7 +9,7 @@
 #include "Stats/Stats.h"
 #include "PhysDerivedDataPublic.h"
 
-#if WITH_PHYSX && WITH_EDITOR
+#if PHYSICS_INTERFACE_PHYSX && WITH_EDITOR
 #include "DerivedDataPluginInterface.h"
 #include "IPhysXCooking.h"
 
diff --git a/Engine/Source/Runtime/Engine/Private/PhysicsEngine/PhysLevel.cpp b/Engine/Source/Runtime/Engine/Private/PhysicsEngine/PhysLevel.cpp
--- a/Engine/Source/Runtime/Engine/Private/PhysicsEngine/PhysLevel.cpp
+++ b/Engine/Source/Runtime/Engine/Private/PhysicsEngine/PhysLevel.cpp
@@ -17,7 +17,7 @@
 #include "PhysicsEngine/PhysicsSettings.h"
 #include "ChaosSolversModule.h"
 
-#if WITH_PHYSX
+#if PHYSICS_INTERFACE_PHYSX
 	#include "PhysicsEngine/PhysXSupport.h"
 #endif
 
@@ -147,11 +147,9 @@ void UWorld::SetupPhysicsTickFunctions(float DeltaSeconds)
 		return;
 	}
 
-#if WITH_PHYSX
 	
 	// When ticking the main scene, clean up any physics engine resources (once a frame)
 	DeferredPhysResourceCleanup();
-#endif
 
 	// Update gravity in case it changed
 	FVector DefaultGravity( 0.f, 0.f, GetGravityZ() );
@@ -236,7 +234,7 @@ void FEndPhysicsTickFunction::ExecuteTick(float DeltaTime, enum ELevelTick TickT
 		Target->FinishPhysicsSim();
 	}
 
-#if WITH_PHYSX
+#if PHYSICS_INTERFACE_PHYSX
 #if PHYSX_MEMORY_VALIDATION
 	static int32 Frequency = 0;
 	if (Frequency++ > 10)
@@ -288,7 +286,6 @@ bool InitGamePhys()
 		PostEngineInitialize();
 	});
 
-#if WITH_PHYSX
 	
 	GPhysCommandHandler = new FPhysCommandHandler();
 	GPreGarbageCollectDelegateHandle = FCoreUObjectDelegates::GetPreGarbageCollectDelegate().AddRaw(GPhysCommandHandler, &FPhysCommandHandler::Flush);
@@ -299,7 +296,6 @@ bool InitGamePhys()
 		DeferredPhysResourceCleanup();
 	});
 	
-#endif // WITH_PHYSX
 
 	// Message to the log that physics is initialised and which interface we are using.
 	UE_LOG(LogInit, Log, TEXT("Physics initialised using underlying interface: %s"), *FPhysicsInterface::GetInterfaceDescription());
@@ -315,14 +311,14 @@ void TermGamePhys()
 		GPostInitHandle.Reset();
 	}
 
-#if WITH_PHYSX
-
+#if PHYSICS_INTERFACE_PHYSX
 	// Do nothing if they were never initialized
 	if(GPhysXFoundation == NULL)
 	{
 		FPhysxSharedData::Terminate();	//early out before TermGamePhysCore so kill this - not sure if this is a real case we even care about
 		return;
 	}
+#endif
 
 	if (GPhysCommandHandler != NULL)
 	{
@@ -331,7 +327,6 @@ void TermGamePhys()
 		delete GPhysCommandHandler;
 		GPhysCommandHandler = NULL;
 	}
-#endif
 
 	TermGamePhysCore();
 }
@@ -342,7 +337,7 @@ void TermGamePhys()
 */
 void DeferredPhysResourceCleanup()
 {
-#if WITH_PHYSX
+#if PHYSICS_INTERFACE_PHYSX
 
 	// Release all tri meshes and reset array
 	for(int32 MeshIdx=0; MeshIdx<GPhysXPendingKillTriMesh.Num(); MeshIdx++)
diff --git a/Engine/Source/Runtime/Engine/Private/PhysicsEngine/PhysSubstepTasks.cpp b/Engine/Source/Runtime/Engine/Private/PhysicsEngine/PhysSubstepTasks.cpp
--- a/Engine/Source/Runtime/Engine/Private/PhysicsEngine/PhysSubstepTasks.cpp
+++ b/Engine/Source/Runtime/Engine/Private/PhysicsEngine/PhysSubstepTasks.cpp
@@ -4,7 +4,7 @@
 #include "PhysicsEngine/PhysicsSettings.h"
 #include "Physics/PhysScene_PhysX.h"
 
-#if WITH_PHYSX
+#if PHYSICS_INTERFACE_PHYSX
 #include "PhysXPublic.h"
 
 uint8* PhysXCompletionTask::GetScratchBufferData()
@@ -39,7 +39,7 @@ struct FSubstepCallbackGuard
 #endif
 };
 
-#if WITH_PHYSX
+#if PHYSICS_INTERFACE_PHYSX
 FPhysSubstepTask::FPhysSubstepTask(PxApexScene * GivenScene, FPhysScene* InPhysScene) :
 	NumSubsteps(0),
 	SubTime(0.f),
@@ -231,7 +231,7 @@ void FPhysSubstepTask::ApplyCustomPhysics(const FPhysTarget& PhysTarget, FBodyIn
 #endif
 }
 
-#if WITH_PHYSX
+#if PHYSICS_INTERFACE_PHYSX
 bool IsKinematicHelper(const PxRigidBody* PRigidBody)
 {
 	const bool bIsKinematic = PRigidBody->getRigidBodyFlags() & PxRigidBodyFlag::eKINEMATIC;
@@ -354,7 +354,7 @@ void FPhysSubstepTask::InterpolateKinematicActor_AssumesLocked(const FPhysTarget
 
 void FPhysSubstepTask::SubstepInterpolation(float InAlpha, float DeltaTime)
 {
-#if WITH_PHYSX
+#if PHYSICS_INTERFACE_PHYSX
 #if WITH_APEX
 	SCOPED_APEX_SCENE_WRITE_LOCK(PAScene);
 	PxScene * PScene = PAScene->getPhysXScene();
@@ -426,7 +426,7 @@ float FPhysSubstepTask::UpdateTime(float UseDelta)
 	return SubTime;
 }
 
-#if WITH_PHYSX
+#if PHYSICS_INTERFACE_PHYSX
 void FPhysSubstepTask::StepSimulation(PhysXCompletionTask * Task)
 {
 	check(SubTime > 0.f);
@@ -449,7 +449,7 @@ void FPhysSubstepTask::SubstepSimulationStart()
 {
 	SCOPE_CYCLE_COUNTER(STAT_TotalPhysicsTime);
 	SCOPE_CYCLE_COUNTER(STAT_SubstepSimulationStart);
-#if WITH_PHYSX
+#if PHYSICS_INTERFACE_PHYSX
 	check(SubTime > 0.f);
 	check(DeltaSeconds > 0.f);
 	
@@ -505,7 +505,7 @@ void FPhysSubstepTask::SubstepSimulationStart()
 
 void FPhysSubstepTask::SubstepSimulationEnd(ENamedThreads::Type CurrentThread, const FGraphEventRef& MyCompletionGraphEvent)
 {
-#if WITH_PHYSX
+#if PHYSICS_INTERFACE_PHYSX
 	CompletionEvent = NULL;
 	if (CurrentSubStep < NumSubsteps)
 	{
diff --git a/Engine/Source/Runtime/Engine/Private/PhysicsEngine/PhysSubstepTasks.h b/Engine/Source/Runtime/Engine/Private/PhysicsEngine/PhysSubstepTasks.h
--- a/Engine/Source/Runtime/Engine/Private/PhysicsEngine/PhysSubstepTasks.h
+++ b/Engine/Source/Runtime/Engine/Private/PhysicsEngine/PhysSubstepTasks.h
@@ -7,7 +7,7 @@
 #include "HAL/ThreadSafeBool.h"
 #include "Physics/PhysicsInterfaceCore.h"
 #include "PhysicsEngine/BodyInstance.h"
-#if WITH_PHYSX
+#if PHYSICS_INTERFACE_PHYSX
 #include "PhysicsEngine/PhysXSupport.h"
 #endif
 
@@ -16,7 +16,7 @@ struct FSimulationScratchBuffer;
 void FinishSceneStat();
 
 //This is only here for now while we transition into substepping
-#if WITH_PHYSX
+#if PHYSICS_INTERFACE_PHYSX
 class PhysXCompletionTask : public PxLightCpuTask
 {
 	FGraphEventRef EventToFire;
@@ -149,7 +149,7 @@ struct FPhysTarget
 	
 };
 
-#if WITH_PHYSX
+#if PHYSICS_INTERFACE_PHYSX
 #if WITH_APEX
 typedef apex::Scene PxApexScene;	//helper typedef so we don't have to use as many ifdefs
 #else
@@ -161,7 +161,7 @@ typedef PxScene PxApexScene;
 class FPhysSubstepTask
 {
 public:
-#if WITH_PHYSX
+#if PHYSICS_INTERFACE_PHYSX
 	FPhysSubstepTask(PxApexScene * GivenScene, FPhysScene* InPhysScene);
 #endif
 
@@ -184,7 +184,7 @@ public:
 
 	void SubstepSimulationStart();
 	void SubstepSimulationEnd(ENamedThreads::Type CurrentThread, const FGraphEventRef& MyCompletionGraphEvent);
-#if WITH_PHYSX
+#if PHYSICS_INTERFACE_PHYSX
 	void StepSimulation(PhysXCompletionTask * Task);
 #endif
 		
@@ -216,7 +216,7 @@ private:
 
 	FPhysScene* PhysScene;
 
-#if WITH_PHYSX
+#if PHYSICS_INTERFACE_PHYSX
 	PxApexScene * PAScene;
 #endif
 
diff --git a/Engine/Source/Runtime/Engine/Private/PhysicsEngine/PhysXCookHelper.cpp b/Engine/Source/Runtime/Engine/Private/PhysicsEngine/PhysXCookHelper.cpp
--- a/Engine/Source/Runtime/Engine/Private/PhysicsEngine/PhysXCookHelper.cpp
+++ b/Engine/Source/Runtime/Engine/Private/PhysicsEngine/PhysXCookHelper.cpp
@@ -5,7 +5,7 @@
 #include "PhysXSupport.h"
 #include "IPhysXCookingModule.h"
 
-#if WITH_PHYSX
+#if PHYSICS_INTERFACE_PHYSX
 
 #include "IPhysXCooking.h"
 
diff --git a/Engine/Source/Runtime/Engine/Private/PhysicsEngine/PhysXSupport.cpp b/Engine/Source/Runtime/Engine/Private/PhysicsEngine/PhysXSupport.cpp
--- a/Engine/Source/Runtime/Engine/Private/PhysicsEngine/PhysXSupport.cpp
+++ b/Engine/Source/Runtime/Engine/Private/PhysicsEngine/PhysXSupport.cpp
@@ -12,7 +12,7 @@
 #include "PhysicalMaterials/PhysicalMaterial.h"
 #include "CustomPhysXPayload.h"
 
-#if WITH_PHYSX
+#if PHYSICS_INTERFACE_PHYSX
 
 #include "PhysXPublic.h"
 #include "PhysicsEngine/ConstraintInstance.h"
diff --git a/Engine/Source/Runtime/Engine/Private/PhysicsEngine/PhysXSupport.h b/Engine/Source/Runtime/Engine/Private/PhysicsEngine/PhysXSupport.h
--- a/Engine/Source/Runtime/Engine/Private/PhysicsEngine/PhysXSupport.h
+++ b/Engine/Source/Runtime/Engine/Private/PhysicsEngine/PhysXSupport.h
@@ -21,7 +21,7 @@ class UPhysicalMaterial;
 struct FCollisionShape;
 class FPhysScene_PhysX;
 
-#if WITH_PHYSX
+#if PHYSICS_INTERFACE_PHYSX
 
 // binary serialization requires 128 byte alignment
 #ifndef PHYSX_SERIALIZATION_ALIGNMENT
diff --git a/Engine/Source/Runtime/Engine/Private/PhysicsEngine/PhysicalAnimationComponent.cpp b/Engine/Source/Runtime/Engine/Private/PhysicsEngine/PhysicalAnimationComponent.cpp
--- a/Engine/Source/Runtime/Engine/Private/PhysicsEngine/PhysicalAnimationComponent.cpp
+++ b/Engine/Source/Runtime/Engine/Private/PhysicsEngine/PhysicalAnimationComponent.cpp
@@ -176,7 +176,7 @@ FTransform UPhysicalAnimationComponent::GetBodyTargetTransform(FName BodyName) c
 {
 	if (SkeletalMeshComponent)
 	{
-#if WITH_PHYSX
+#if PHYSICS_INTERFACE_PHYSX
 		for (int32 DataIdx = 0; DataIdx < DriveData.Num(); ++DataIdx)
 		{
 			const FPhysicalAnimationData& PhysAnimData = DriveData[DataIdx];
@@ -199,6 +199,8 @@ FTransform UPhysicalAnimationComponent::GetBodyTargetTransform(FName BodyName) c
 				break;
 			}
 		}
+#else
+		CHAOS_ENSURE(false);
 #endif
 
 		// if body isn't controlled by physical animation, just return the body position
@@ -270,7 +272,7 @@ void UPhysicalAnimationComponent::UpdateTargetActors(ETeleportType TeleportType)
 		// various anim ticks, before buffers are flipped (which happens in the skel mesh component's post-physics tick)
 		const TArray<FTransform>& SpaceBases = SkeletalMeshComponent->GetEditableComponentSpaceTransforms();
 
-#if WITH_PHYSX
+#if PHYSICS_INTERFACE_PHYSX
 		FPhysicsCommand::ExecuteWrite(SkeletalMeshComponent, [&]()
 		{
 			TArray<FTransform> LocalTransforms = SkeletalMeshComponent->GetBoneSpaceTransforms();
@@ -295,6 +297,8 @@ void UPhysicalAnimationComponent::UpdateTargetActors(ETeleportType TeleportType)
 				}
 			}
 		});
+#else
+		CHAOS_ENSURE(false);
 #endif
 	}
 }
@@ -378,7 +382,7 @@ void UPhysicalAnimationComponent::UpdatePhysicsEngineImp()
 					{
 #if WITH_CHAOS || WITH_IMMEDIATE_PHYSX
                         CHAOS_ENSURE(false);
-#else
+#elif PHYSICS_INTERFACE_PHYSX
 						if (PxRigidActor* PRigidActor = FPhysicsInterface_PhysX::GetPxRigidActor_AssumesLocked(ChildBody->ActorHandle))
 						{
 							ConstraintInstance->SetRefFrame(EConstraintFrame::Frame1, FTransform::Identity);
@@ -447,7 +451,7 @@ void UPhysicalAnimationComponent::SetStrengthMultiplyer(float InStrengthMultiply
 
 void UPhysicalAnimationComponent::ReleasePhysicsEngine()
 {
-#if WITH_PHYSX
+#if PHYSICS_INTERFACE_PHYSX
 	// #PHYS2 On shutdown, SkelMeshComp is null, so we can't lock using that, need to lock based on scene from body
 	//FPhysicsCommand::ExecuteWrite(SkeletalMeshComponent, [&]()
 	//{
@@ -486,7 +490,7 @@ void UPhysicalAnimationComponent::DebugDraw(FPrimitiveDrawInterface* PDI) const
 {
 	for (const FPhysicalAnimationInstanceData& PhysAnimData : RuntimeInstanceData)
 	{
-#if WITH_PHYSX
+#if PHYSICS_INTERFACE_PHYSX
 		if (PhysAnimData.TargetActor)
 		{
 			PDI->DrawPoint(P2UVector(PhysAnimData.TargetActor->getGlobalPose().p), TargetActorColor, 3.f, SDPG_World);
diff --git a/Engine/Source/Runtime/Engine/Private/PhysicsEngine/PhysicalMaterial.cpp b/Engine/Source/Runtime/Engine/Private/PhysicsEngine/PhysicalMaterial.cpp
--- a/Engine/Source/Runtime/Engine/Private/PhysicsEngine/PhysicalMaterial.cpp
+++ b/Engine/Source/Runtime/Engine/Private/PhysicsEngine/PhysicalMaterial.cpp
@@ -12,7 +12,7 @@
 #include "PhysicalMaterials/PhysicalMaterialPropertyBase.h"
 #include "PhysicsEngine/PhysicsSettings.h"
 
-#if WITH_PHYSX
+#if PHYSICS_INTERFACE_PHYSX
 	#include "PhysicsEngine/PhysXSupport.h"
 #endif // WITH_PHYSX
 
@@ -31,7 +31,7 @@ UPhysicalMaterial::UPhysicalMaterial(const FObjectInitializer& ObjectInitializer
 	DestructibleDamageThresholdScale = 1.0f;
 	TireFrictionScale = 1.0f;
 	bOverrideFrictionCombineMode = false;
-#if WITH_PHYSX
+#if PHYSICS_INTERFACE_PHYSX
 	PhysxUserData = FPhysxUserData(this);
 #endif
 }
diff --git a/Engine/Source/Runtime/Engine/Private/PhysicsEngine/PhysicsHandleComponent.cpp b/Engine/Source/Runtime/Engine/Private/PhysicsEngine/PhysicsHandleComponent.cpp
--- a/Engine/Source/Runtime/Engine/Private/PhysicsEngine/PhysicsHandleComponent.cpp
+++ b/Engine/Source/Runtime/Engine/Private/PhysicsEngine/PhysicsHandleComponent.cpp
@@ -7,7 +7,7 @@
 #include "PhysicsPublic.h"
 #include "Physics/PhysicsInterfaceCore.h"
 
-#if WITH_PHYSX
+#if PHYSICS_INTERFACE_PHYSX
 	#include "PhysXPublic.h"
 #endif // WITH_PHYSX
 
@@ -34,7 +34,7 @@ void UPhysicsHandleComponent::OnUnregister()
 		ReleaseComponent();
 	}
 
-#if WITH_PHYSX
+#if PHYSICS_INTERFACE_PHYSX
 	if(HandleData)
 	{
 		check(KinActorData);
@@ -195,7 +195,9 @@ void UPhysicsHandleComponent::GrabComponentImp(UPrimitiveComponent* InComponent,
 		}
 	});
 	
-#endif // WITH_PHYSX
+#elif WITH_CHAOS
+	CHAOS_ENSURE(false);
+#endif
 
 	GrabbedComponent = InComponent;
 	GrabbedBoneName = InBoneName;
@@ -203,7 +205,7 @@ void UPhysicsHandleComponent::GrabComponentImp(UPrimitiveComponent* InComponent,
 
 void UPhysicsHandleComponent::UpdateDriveSettings()
 {
-#if WITH_PHYSX
+#if PHYSICS_INTERFACE_PHYSX
 	if(HandleData != nullptr)
 	{
 		if (bSoftLinearConstraint)
@@ -228,7 +230,7 @@ void UPhysicsHandleComponent::ReleaseComponent()
 {
 	if(GrabbedComponent)
 	{
-#if WITH_PHYSX
+#if PHYSICS_INTERFACE_PHYSX
 		if(HandleData)
 		{
 			check(KinActorData);
@@ -287,7 +289,7 @@ void UPhysicsHandleComponent::UpdateHandleTransform(const FTransform& NewTransfo
 		return;
 	}
 
-#if WITH_PHYSX
+#if PHYSICS_INTERFACE_PHYSX
 	bool bChangedPosition = true;
 	bool bChangedRotation = true;
 
diff --git a/Engine/Source/Runtime/Engine/Private/UnrealEngine.cpp b/Engine/Source/Runtime/Engine/Private/UnrealEngine.cpp
--- a/Engine/Source/Runtime/Engine/Private/UnrealEngine.cpp
+++ b/Engine/Source/Runtime/Engine/Private/UnrealEngine.cpp
@@ -129,6 +129,7 @@ UnrealEngine.cpp: Implements the UEngine class and helpers.
 #include "HAL/LowLevelMemTracker.h"
 #include "HAL/PlatformApplicationMisc.h"
 #include "DynamicResolutionState.h"
+#include "Chaos/TriangleMeshImplicitObject.h"
 
 #include "Particles/Spawn/ParticleModuleSpawn.h"
 #include "Particles/TypeData/ParticleModuleTypeDataMesh.h"
@@ -5350,12 +5351,17 @@ bool UEngine::HandleListStaticMeshesCommand(const TCHAR* Cmd, FOutputDevice& Ar)
 		int32		VertexCountCollision = 0;
 		if(Mesh->BodySetup)
 		{
-#if WITH_PHYSX
+#if PHYSICS_INTERFACE_PHYSX
 			// Count PhysX trimesh mem usage
 			for (physx::PxTriangleMesh* TriMesh : Mesh->BodySetup->TriMeshes)
 			{
 				VertexCountCollision += TriMesh->getNbVertices();
 			}
+#elif WITH_CHAOS
+			for (auto& TriMesh : Mesh->BodySetup->ChaosTriMeshes)
+			{
+				VertexCountCollision += TriMesh->Particles().Size();
+			}
 #endif
 		}
 
@@ -5564,12 +5570,17 @@ bool UEngine::HandleListSkeletalMeshesCommand(const TCHAR* Cmd, FOutputDevice& A
 		int32 VertexCountCollision = 0;
 		if (Mesh->BodySetup)
 		{
-#if WITH_PHYSX
+#if PHYSICS_INTERFACE_PHYSX
 			// Count PhysX trimesh mem usage
 			for (physx::PxTriangleMesh* TriMesh : Mesh->BodySetup->TriMeshes)
 			{
 				VertexCountCollision += TriMesh->getNbVertices();
 			}
+#elif WITH_CHAOS
+			for (auto& TriMesh : Mesh->BodySetup->ChaosTriMeshes)
+			{
+				VertexCountCollision += TriMesh->Particles().Size();
+			}
 #endif
 		}
 
diff --git a/Engine/Source/Runtime/Engine/Public/AI/NavigationSystemHelpers.h b/Engine/Source/Runtime/Engine/Public/AI/NavigationSystemHelpers.h
--- a/Engine/Source/Runtime/Engine/Public/AI/NavigationSystemHelpers.h
+++ b/Engine/Source/Runtime/Engine/Public/AI/NavigationSystemHelpers.h
@@ -16,7 +16,7 @@ struct FNavHeightfieldSamples;
 
 template<typename InElementType> class TNavStatArray;
 
-#if WITH_PHYSX
+#if PHYSICS_INTERFACE_PHYSX
 namespace physx
 {
 	class PxTriangleMesh;
@@ -38,7 +38,7 @@ namespace Chaos
 struct FNavigableGeometryExport
 {
 	virtual ~FNavigableGeometryExport() {}
-#if WITH_PHYSX
+#if PHYSICS_INTERFACE_PHYSX
 	virtual void ExportPxTriMesh16Bit(physx::PxTriangleMesh const * const TriMesh, const FTransform& LocalToWorld) = 0;
 	virtual void ExportPxTriMesh32Bit(physx::PxTriangleMesh const * const TriMesh, const FTransform& LocalToWorld) = 0;
 	virtual void ExportPxConvexMesh(physx::PxConvexMesh const * const ConvexMesh, const FTransform& LocalToWorld) = 0;
diff --git a/Engine/Source/Runtime/Engine/Public/PhysXInterfaceWrapper.h b/Engine/Source/Runtime/Engine/Public/PhysXInterfaceWrapper.h
--- a/Engine/Source/Runtime/Engine/Public/PhysXInterfaceWrapper.h
+++ b/Engine/Source/Runtime/Engine/Public/PhysXInterfaceWrapper.h
@@ -1,7 +1,7 @@
 // Copyright Epic Games, Inc. All Rights Reserved.
 #pragma once
 
-#if WITH_PHYSX
+#if PHYSICS_INTERFACE_PHYSX
 #include "PhysXInterfaceWrapperCore.h"
 #include "PhysXPublic.h"
 #include "Physics/PhysicsInterfaceUtils.h"
diff --git a/Engine/Source/Runtime/Engine/Public/PhysXPublic.h b/Engine/Source/Runtime/Engine/Public/PhysXPublic.h
--- a/Engine/Source/Runtime/Engine/Public/PhysXPublic.h
+++ b/Engine/Source/Runtime/Engine/Public/PhysXPublic.h
@@ -8,7 +8,7 @@
 
 #include "CoreMinimal.h"
 
-#if WITH_PHYSX
+#if PHYSICS_INTERFACE_PHYSX
 #include "PhysXPublicCore.h"
 
 /** Calculates correct impulse at the body's center of mass and adds the impulse to the body. */
diff --git a/Engine/Source/Runtime/Engine/Public/Physics/Experimental/PhysInterface_Chaos.h b/Engine/Source/Runtime/Engine/Public/Physics/Experimental/PhysInterface_Chaos.h
--- a/Engine/Source/Runtime/Engine/Public/Physics/Experimental/PhysInterface_Chaos.h
+++ b/Engine/Source/Runtime/Engine/Public/Physics/Experimental/PhysInterface_Chaos.h
@@ -374,10 +374,7 @@ public:
 
 	static bool ExecPhysCommands(const TCHAR* Cmd, FOutputDevice* Ar, UWorld* InWorld);
 
-	static void CalculateTMassPropertiesFromShapeCollection(Chaos::TMassProperties<float, 3>& OutProperties, const TArray<FPhysicsShapeHandle>& InShapes, float InDensityKGPerCM);
-#if WITH_PHYSX
-    static void CalculateMassPropertiesFromShapeCollection(physx::PxMassProperties& OutProperties, const TArray<FPhysicsShapeHandle>& InShapes, float InDensityKGPerCM);
-#endif
+	static void CalculateMassPropertiesFromShapeCollection(Chaos::TMassProperties<float, 3>& OutProperties, const TArray<FPhysicsShapeHandle>& InShapes, float InDensityKGPerCM);
 
 	// Shape interface functions
 	static FPhysicsShapeHandle CreateShape(physx::PxGeometry* InGeom, bool bSimulation = true, bool bQuery = true, UPhysicalMaterial* InSimpleMaterial = nullptr, TArray<UPhysicalMaterial*>* InComplexMaterials = nullptr);
diff --git a/Engine/Source/Runtime/Engine/Public/Physics/PhysicsFiltering.h b/Engine/Source/Runtime/Engine/Public/Physics/PhysicsFiltering.h
--- a/Engine/Source/Runtime/Engine/Public/Physics/PhysicsFiltering.h
+++ b/Engine/Source/Runtime/Engine/Public/Physics/PhysicsFiltering.h
@@ -143,7 +143,7 @@ inline void UpdateMaskFilter(uint32& Word3, FMaskFilter NewMaskFilter)
 	Word3 |= uint32(NewMaskFilter) << (32 - NumExtraFilterBits);
 }
 
-#if WITH_PHYSX
+#if PHYSICS_INTERFACE_PHYSX
 	extern ENGINE_API PxSimulationFilterShader GSimulationFilterShader;
 #endif
 
diff --git a/Engine/Source/Runtime/Engine/Public/PhysicsPublic.h b/Engine/Source/Runtime/Engine/Public/PhysicsPublic.h
--- a/Engine/Source/Runtime/Engine/Public/PhysicsPublic.h
+++ b/Engine/Source/Runtime/Engine/Public/PhysicsPublic.h
@@ -54,7 +54,11 @@ DECLARE_DWORD_COUNTER_STAT_EXTERN(TEXT("Cloth Actor Count"), STAT_NumCloths, STA
 DECLARE_DWORD_COUNTER_STAT_EXTERN(TEXT("Simulated Cloth Verts"), STAT_NumClothVerts, STATGROUP_Physics, ENGINE_API);
 
 #define WITH_PHYSX_VEHICLES WITH_PHYSX && PHYSICS_INTERFACE_PHYSX
-#if WITH_PHYSX
+
+/** Pointer to PhysX Command Handler */
+extern ENGINE_API class FPhysCommandHandler* GPhysCommandHandler;
+
+#if PHYSICS_INTERFACE_PHYSX
 
 namespace physx
 {
@@ -107,9 +111,6 @@ using namespace nvidia;
 /** Pointer to PhysX cooking object */
 extern ENGINE_API PxCooking*			GPhysXCooking;
 
-/** Pointer to PhysX Command Handler */
-extern ENGINE_API class FPhysCommandHandler* GPhysCommandHandler;
-
 namespace NvParameterized
 {
 	class Interface;
@@ -201,7 +202,7 @@ public:
 	void ENGINE_API DeferredRelease(apex::ApexInterface* ApexInterface);
 #endif
 
-#if WITH_PHYSX
+#if PHYSICS_INTERFACE_PHYSX
 	void ENGINE_API DeferredRelease(physx::PxScene * PScene);
 	void ENGINE_API DeferredDeleteSimEventCallback(physx::PxSimulationEventCallback* SimEventCallback);
 	void ENGINE_API DeferredDeleteContactModifyCallback(FContactModifyCallback* ContactModifyCallback);
@@ -221,7 +222,7 @@ private:
 			apex::ApexInterface * ApexInterface;
 			apex::DestructibleActor * DestructibleActor;
 #endif
-#if WITH_PHYSX
+#if PHYSICS_INTERFACE_PHYSX
 			physx::PxScene* PScene;
 			physx::PxCpuDispatcher* CPUDispatcher;
 			physx::PxSimulationEventCallback* SimEventCallback;
@@ -263,7 +264,7 @@ FORCEINLINE bool PhysSingleThreadedMode()
 	return false;
 }
 
-#if WITH_PHYSX
+#if PHYSICS_INTERFACE_PHYSX
 /** Struct used for passing info to the PhysX shader */
 
 struct FPhysSceneShaderInfo
diff --git a/Engine/Source/Runtime/Engine/Public/PhysxUserData.h b/Engine/Source/Runtime/Engine/Public/PhysxUserData.h
--- a/Engine/Source/Runtime/Engine/Public/PhysxUserData.h
+++ b/Engine/Source/Runtime/Engine/Public/PhysxUserData.h
@@ -1,6 +1,7 @@
 // Copyright Epic Games, Inc. All Rights Reserved.
 #pragma once
 
+#if PHYSICS_INTERFACE_PHYSX
 #include "CoreMinimal.h"
 #include "EngineDefines.h"
 #include "Physics/PhysicsInterfaceDeclares.h"
@@ -70,4 +71,5 @@ template <> FORCEINLINE void FPhysxUserData::Set(void* UserData, FPhysScene* Pay
 template <> FORCEINLINE void FPhysxUserData::Set(void* UserData, FConstraintInstance* Payload)		{ check(UserData); ((FPhysxUserData*)UserData)->Type = EPhysxUserDataType::ConstraintInstance; ((FPhysxUserData*)UserData)->Payload = Payload; }
 template <> FORCEINLINE void FPhysxUserData::Set(void* UserData, UPrimitiveComponent* Payload)		{ check(UserData); ((FPhysxUserData*)UserData)->Type = EPhysxUserDataType::PrimitiveComponent; ((FPhysxUserData*)UserData)->Payload = Payload; }
 template <> FORCEINLINE void FPhysxUserData::Set(void* UserData, FKShapeElem* Payload)	{ check(UserData); ((FPhysxUserData*)UserData)->Type = EPhysxUserDataType::AggShape; ((FPhysxUserData*)UserData)->Payload = Payload; }
-template <> FORCEINLINE void FPhysxUserData::Set(void* UserData, FCustomPhysXPayload* Payload) { check(UserData); ((FPhysxUserData*)UserData)->Type = EPhysxUserDataType::CustomPayload; ((FPhysxUserData*)UserData)->Payload = Payload; }
\ No newline at end of file
+template <> FORCEINLINE void FPhysxUserData::Set(void* UserData, FCustomPhysXPayload* Payload) { check(UserData); ((FPhysxUserData*)UserData)->Type = EPhysxUserDataType::CustomPayload; ((FPhysxUserData*)UserData)->Payload = Payload; }
+#endif 
\ No newline at end of file
diff --git a/Engine/Source/Runtime/Experimental/GeometryCollectionEngine/Private/GeometryCollection/PhysicsAssetSimulation.cpp b/Engine/Source/Runtime/Experimental/GeometryCollectionEngine/Private/GeometryCollection/PhysicsAssetSimulation.cpp
--- a/Engine/Source/Runtime/Experimental/GeometryCollectionEngine/Private/GeometryCollection/PhysicsAssetSimulation.cpp
+++ b/Engine/Source/Runtime/Experimental/GeometryCollectionEngine/Private/GeometryCollection/PhysicsAssetSimulation.cpp
@@ -279,6 +279,7 @@ void FPhysicsAssetSimulationUtil::BuildParams(const UObject* Caller, const AActo
 					Bounds += Points[i];
 				}
 
+#if PHYSICS_INTERFACE_PHYSX
 				if (physx::PxConvexMesh *ConvexMesh = Elem.GetConvexMesh())
 				{
 					ConvexMesh->acquireReference();
@@ -364,6 +365,10 @@ void FPhysicsAssetSimulationUtil::BuildParams(const UObject* Caller, const AActo
 								"ADDED (LEVELSET) SPHERE"), Caller, &OwningActor->GetName()[0]);
 					}
 				}
+#elif WITH_CHAOS
+				// TODO: Use Chaos convex.
+				CHAOS_ENSURE(false);
+#endif
 			}
 		}
 		Params.BoneHierarchy.Add(MoveTemp(AnalyticShapeGroup));
diff --git a/Engine/Source/Runtime/Landscape/Private/LandscapeCollision.cpp b/Engine/Source/Runtime/Landscape/Private/LandscapeCollision.cpp
--- a/Engine/Source/Runtime/Landscape/Private/LandscapeCollision.cpp
+++ b/Engine/Source/Runtime/Landscape/Private/LandscapeCollision.cpp
@@ -44,9 +44,8 @@
 #include "Materials/MaterialInstanceConstant.h"
 #include "Physics/PhysicsInterfaceCore.h"
 #include "Physics/PhysicsInterfaceUtils.h"
-#include "Containers/StringConv.h"
 
-#if WITH_EDITOR && WITH_PHYSX
+#if WITH_EDITOR && PHYSICS_INTERFACE_PHYSX
 	#include "IPhysXCooking.h"
 #endif
 
@@ -401,11 +400,6 @@ void ULandscapeHeightfieldCollisionComponent::OnCreatePhysicsState()
 				Params.InitialTM.SetScale3D(FVector(0));
 				Params.bQueryOnly = true;
 				Params.Scene = GetWorld()->GetPhysicsScene();
-#if USE_BODYINSTANCE_DEBUG_NAMES
-				const FString DebugName = FString::Printf(TEXT("LandscapeCollision: '%s' "), *GetPathName());
-				BodyInstance.CharDebugName = MakeShareable(new TArray<ANSICHAR>(StringToArray<ANSICHAR>(*DebugName, DebugName.Len() + 1)));
-				Params.DebugName = BodyInstance.CharDebugName->GetData();
-#endif
 				FPhysicsActorHandle PhysHandle;
 				FPhysicsInterface::CreateActor(Params, PhysHandle);
 
@@ -736,7 +730,7 @@ void ULandscapeHeightfieldCollisionComponent::SpeculativelyLoadAsyncDDCCollsionD
 #endif
 }
 
-#if WITH_PHYSX
+#if PHYSICS_INTERFACE_PHYSX
 TArray<PxHeightFieldSample> ConvertHeightfieldDataForPhysx(const ULandscapeHeightfieldCollisionComponent* const Component, const int32 CollisionSizeVerts, const bool bIsMirrored, const uint16* Heights, const bool bUseDefMaterial, const uint8* DominantLayers, UPhysicalMaterial* const DefMaterial, TArray<UPhysicalMaterial*> &InOutMaterials)
 {
 	const int32 NumSamples = FMath::Square(CollisionSizeVerts);
@@ -1022,7 +1016,7 @@ bool ULandscapeMeshCollisionComponent::CookCollisionData(const FName& Format, bo
 		return true;
 	}
 
-#if WITH_PHYSX
+#if PHYSICS_INTERFACE_PHYSX
 	COOK_STAT(auto Timer = LandscapeCollisionCookStats::MeshUsageStats.TimeSyncWork());
 	// we have 2 versions of collision objects
 	const int32 CookedDataIndex = bUseDefMaterial ? 0 : 1;
@@ -1219,14 +1213,16 @@ bool ULandscapeMeshCollisionComponent::CookCollisionData(const FName& Format, bo
 
 	return Result;
 
-#endif // WITH_PHYSX
+#elif WITH_CHAOS
+CHAOS_ENSURE(false);
+#endif
 	return false;
 }
 #endif //WITH_EDITOR
 
 void ULandscapeMeshCollisionComponent::CreateCollisionObject()
 {
-#if WITH_PHYSX	
+#if PHYSICS_INTERFACE_PHYSX
 	// If we have not created a heightfield yet - do it now.
 	if (!IsValidRef(MeshRef))
 	{
@@ -1311,7 +1307,7 @@ void ULandscapeMeshCollisionComponent::OnCreatePhysicsState()
 
 	if (!BodyInstance.IsValidBodyInstance())
 	{
-#if WITH_PHYSX
+#if PHYSICS_INTERFACE_PHYSX
 		// This will do nothing, because we create trimesh at component PostLoad event, unless we destroyed it explicitly
 		CreateCollisionObject();
 
@@ -2047,7 +2043,7 @@ void ULandscapeHeightfieldCollisionComponent::PrepareGeometryExportSync()
 bool ULandscapeMeshCollisionComponent::DoCustomNavigableGeometryExport(FNavigableGeometryExport& GeomExport) const
 {
 	check(IsInGameThread());
-#if WITH_PHYSX
+#if PHYSICS_INTERFACE_PHYSX
 	if (IsValidRef(MeshRef) && MeshRef->RBTriangleMesh != nullptr)
 	{
 		FTransform MeshToW = GetComponentTransform();
@@ -2062,7 +2058,9 @@ bool ULandscapeMeshCollisionComponent::DoCustomNavigableGeometryExport(FNavigabl
 			GeomExport.ExportPxTriMesh32Bit(MeshRef->RBTriangleMesh, MeshToW);
 		}
 	}
-#endif// WITH_PHYSX
+#elif WITH_CHAOS
+	CHAOS_ENSURE(false);
+#endif
 	return false;
 }
 
diff --git a/Engine/Source/Runtime/NavigationSystem/Private/NavMesh/RecastNavMeshGenerator.cpp b/Engine/Source/Runtime/NavigationSystem/Private/NavMesh/RecastNavMeshGenerator.cpp
--- a/Engine/Source/Runtime/NavigationSystem/Private/NavMesh/RecastNavMeshGenerator.cpp
+++ b/Engine/Source/Runtime/NavigationSystem/Private/NavMesh/RecastNavMeshGenerator.cpp
@@ -221,7 +221,7 @@ struct FRecastGeometryExport : public FNavigableGeometryExport
 	TNavStatArray<int32> IndexBuffer;
 	FWalkableSlopeOverride SlopeOverride;
 
-#if WITH_PHYSX
+#if PHYSICS_INTERFACE_PHYSX
 	virtual void ExportPxTriMesh16Bit(physx::PxTriangleMesh const * const TriMesh, const FTransform& LocalToWorld) override;
 	virtual void ExportPxTriMesh32Bit(physx::PxTriangleMesh const * const TriMesh, const FTransform& LocalToWorld) override;
 	virtual void ExportPxConvexMesh(physx::PxConvexMesh const * const ConvexMesh, const FTransform& LocalToWorld) override;
@@ -342,7 +342,7 @@ static void StoreCollisionCache(FRecastGeometryExport& GeomExport)
 	FMemory::Memcpy(RawMemory + HeaderSize + CoordsSize, GeomExport.IndexBuffer.GetData(), IndicesSize);
 }
 
-#if WITH_PHYSX
+#if PHYSICS_INTERFACE_PHYSX
 /** exports PxConvexMesh as trimesh */
 void ExportPxConvexMesh(PxConvexMesh const * const ConvexMesh, const FTransform& LocalToWorld,
 						TNavStatArray<float>& VertexBuffer, TNavStatArray<int32>& IndexBuffer,
@@ -784,7 +784,7 @@ void ExportHeightFieldSlice(const FNavHeightfieldSamples& PrefetchedHeightfieldS
 {
 	QUICK_SCOPE_CYCLE_COUNTER(STAT_NavMesh_ExportHeightFieldSlice);
 
-#if WITH_PHYSX
+#if PHYSICS_INTERFACE_PHYSX
 	static const uint32 SizeOfPx = sizeof(physx::PxI16);
 	static const uint32 SizeOfHeight = PrefetchedHeightfieldSamples.Heights.GetTypeSize();
 	ensure(SizeOfPx == SizeOfHeight);
@@ -1325,7 +1325,7 @@ void ExportVertexSoup(const TArray<FVector>& VertexSoup, TNavStatArray<float>& V
 
 } // namespace RecastGeometryExport
 
-#if WITH_PHYSX
+#if PHYSICS_INTERFACE_PHYSX
 void FRecastGeometryExport::ExportPxTriMesh16Bit(physx::PxTriangleMesh const * const TriMesh, const FTransform& LocalToWorld)
 {
 	RecastGeometryExport::ExportPxTriMesh<PxU16>(TriMesh, LocalToWorld, VertexBuffer, IndexBuffer, Data->Bounds);
diff --git a/Engine/Source/Runtime/NavigationSystem/Private/NavigationSystemTypes.cpp b/Engine/Source/Runtime/NavigationSystem/Private/NavigationSystemTypes.cpp
--- a/Engine/Source/Runtime/NavigationSystem/Private/NavigationSystemTypes.cpp
+++ b/Engine/Source/Runtime/NavigationSystem/Private/NavigationSystemTypes.cpp
@@ -233,7 +233,7 @@ namespace NavigationHelper
 
 	bool IsBodyNavigationRelevant(const UBodySetup& BodySetup)
 	{
-#if WITH_PHYSX
+#if PHYSICS_INTERFACE_PHYSX
 		const bool bBodyHasGeometry = (BodySetup.AggGeom.GetElementCount() > 0 || BodySetup.TriMeshes.Num() > 0);
 #elif WITH_CHAOS
 		const bool bBodyHasGeometry = (BodySetup.AggGeom.GetElementCount() > 0 || BodySetup.ChaosTriMeshes.Num() > 0);
diff --git a/Engine/Source/Runtime/PhysXCooking/Private/PhysXCooking.cpp b/Engine/Source/Runtime/PhysXCooking/Private/PhysXCooking.cpp
--- a/Engine/Source/Runtime/PhysXCooking/Private/PhysXCooking.cpp
+++ b/Engine/Source/Runtime/PhysXCooking/Private/PhysXCooking.cpp
@@ -2,7 +2,7 @@
 
 #include "PhysXCooking.h"
 
-#if WITH_PHYSX
+#if PHYSICS_INTERFACE_PHYSX
 
 #include "Serialization/MemoryWriter.h"
 #include "Modules/ModuleManager.h"
diff --git a/Engine/Source/Runtime/PhysXCooking/Public/PhysXCooking.h b/Engine/Source/Runtime/PhysXCooking/Public/PhysXCooking.h
--- a/Engine/Source/Runtime/PhysXCooking/Public/PhysXCooking.h
+++ b/Engine/Source/Runtime/PhysXCooking/Public/PhysXCooking.h
@@ -1,11 +1,11 @@
 // Copyright Epic Games, Inc. All Rights Reserved.
 
 #pragma once
+#include "IPhysXCookingModule.h"
 
-#if WITH_PHYSX
+#if PHYSICS_INTERFACE_PHYSX
 
 #include "IPhysXCooking.h"
-#include "IPhysXCookingModule.h"
 
 namespace physx
 {
diff --git a/Engine/Source/Runtime/PhysicsCore/Private/PhysTestSerializer.cpp b/Engine/Source/Runtime/PhysicsCore/Private/PhysTestSerializer.cpp
--- a/Engine/Source/Runtime/PhysicsCore/Private/PhysTestSerializer.cpp
+++ b/Engine/Source/Runtime/PhysicsCore/Private/PhysTestSerializer.cpp
@@ -4,11 +4,14 @@
 
 #include "PhysTestSerializer.h"
 
-#if WITH_PHYSX
+
+#if PHYSICS_INTERFACE_PHYSX
 #include "PhysXIncludes.h"
 #include "PhysXSupportCore.h"
+#include "PhysXToChaosUtil.h"
 #endif
 
+#include "PhysicsCore.h"
 #include "Chaos/PBDRigidsEvolution.h"
 #include "Chaos/PBDRigidParticles.h"
 #include "Chaos/Box.h"
@@ -17,10 +20,6 @@
 
 using namespace Chaos;
 
-#if WITH_PHYSX
-#include "PhysXToChaosUtil.h"
-#endif
-
 #include "PhysicsPublicCore.h"
 #include "PhysicsCore.h"
 #include "HAL/FileManager.h"
@@ -93,8 +92,14 @@
 
 	if (Ar.IsLoading())
 	{
+#if PHYSICS_INTERFACE_PHYSX
 		CreatePhysXData();
+#endif
+
+#if 0
 		CreateChaosData();
+#endif
+
 		Ar.SetContext(MoveTemp(ChaosContext));	//make sure any context we created during load is used for sqcapture
 	}
 
@@ -114,9 +119,22 @@
 	ChaosContext = Ar.StealContext();
 }
 
+void FPhysTestSerializer::SetPhysicsData(Chaos::FPBDRigidsEvolutionGBF& Evolution)
+{
+	bDiskDataIsChaos = true;
+	Data.Empty();
+	FMemoryWriter Ar(Data);
+	FChaosArchive ChaosAr(Ar);
+	Evolution.Serialize(ChaosAr);
+	ChaosContext = ChaosAr.StealContext();
+	ArchiveVersion = Ar.GetCustomVersions();
+}
+
+
+
+#if PHYSICS_INTERFACE_PHYSX
 void FPhysTestSerializer::SetPhysicsData(physx::PxScene& Scene)
 {
-#if WITH_PHYSX
 	check(AlignedDataHelper == nullptr || &Scene != AlignedDataHelper->PhysXScene);
 
 	PxSerializationRegistry* Registry = PxSerialization::createSerializationRegistry(*GPhysXSDK);
@@ -141,49 +159,10 @@
 	Registry->release();
 
 	bDiskDataIsChaos = false;
-#endif
 }
 
-void FPhysTestSerializer::SetPhysicsData(Chaos::FPBDRigidsEvolutionGBF& Evolution)
-{
-	bDiskDataIsChaos = true;
-	Data.Empty();
-	FMemoryWriter Ar(Data);
-	FChaosArchive ChaosAr(Ar);
-	Evolution.Serialize(ChaosAr);
-	ChaosContext = ChaosAr.StealContext();
-	ArchiveVersion = Ar.GetCustomVersions();
-}
-
-#if WITH_PHYSX
-FPhysTestSerializer::FPhysXSerializerData::~FPhysXSerializerData()
-{
-	if (PhysXScene)
-	{
-		//release all resources the collection created (calling release on the collection is not enough)
-		const uint32 NumObjects = Collection->getNbObjects();
-		TArray<PxBase*> Objects;
-		Objects.AddUninitialized(NumObjects);
-		Collection->getObjects(Objects.GetData(), NumObjects);
-		for (PxBase* Obj : Objects)
-		{
-			if (Obj->isReleasable())
-			{
-				Obj->release();
-			}
-		}
-
-		Collection->release();
-		Registry->release();
-		PhysXScene->release();
-	}
-	FMemory::Free(Data);
-}
-#endif
-
 void FPhysTestSerializer::CreatePhysXData()
 {
-#if WITH_PHYSX
 	if (bDiskDataIsChaos == false)	//For the moment we don't support chaos to physx direction
 	{
 		{
@@ -199,10 +178,8 @@
 		AlignedDataHelper->Collection = PxSerialization::createCollectionFromBinary(AlignedDataHelper->Data, *AlignedDataHelper->Registry);
 		AlignedDataHelper->PhysXScene->addCollection(*AlignedDataHelper->Collection);
 	}
-#endif
 }
 
-#if WITH_PHYSX
 physx::PxBase* FPhysTestSerializer::FindObject(uint64 Id)
 {
 	if (!AlignedDataHelper)
@@ -212,13 +189,43 @@
 
 	physx::PxBase* Ret = AlignedDataHelper->Collection->find(Id);
 	ensure(Ret);
+#if 0
+		CreateChaosData();
+#endif
 	return Ret;
 }
+
+FPhysTestSerializer::FPhysXSerializerData::~FPhysXSerializerData()
+{
+	if (PhysXScene)
+	{
+		//release all resources the collection created (calling release on the collection is not enough)
+		const uint32 NumObjects = Collection->getNbObjects();
+		TArray<PxBase*> Objects;
+		Objects.AddUninitialized(NumObjects);
+		Collection->getObjects(Objects.GetData(), NumObjects);
+		for (PxBase* Obj : Objects)
+		{
+			if (Obj->isReleasable())
+			{
+				Obj->release();
+			}
+		}
+
+		Collection->release();
+		Registry->release();
+		PhysXScene->release();
+	}
+	FMemory::Free(Data);
+}
+
+
 #endif
 
+#if 0
+
 void FPhysTestSerializer::CreateChaosData()
 {
-#if WITH_PHYSX
 	if (bDiskDataIsChaos == false)
 	{
 		if (bChaosDataReady)
@@ -347,5 +354,5 @@
 		ChaosContext = ChaosAr.StealContext();
 	}
 	bChaosDataReady = true;
-#endif
 }
+#endif
diff --git a/Engine/Source/Runtime/PhysicsCore/Private/PhysXPublicCore.cpp b/Engine/Source/Runtime/PhysicsCore/Private/PhysXPublicCore.cpp
--- a/Engine/Source/Runtime/PhysicsCore/Private/PhysXPublicCore.cpp
+++ b/Engine/Source/Runtime/PhysicsCore/Private/PhysXPublicCore.cpp
@@ -1,6 +1,6 @@
 // Copyright Epic Games, Inc. All Rights Reserved.
 
-#if WITH_PHYSX
+#if PHYSICS_INTERFACE_PHYSX
 
 #include "PhysXPublicCore.h"
 #include "PhysicsPublicCore.h"
diff --git a/Engine/Source/Runtime/PhysicsCore/Private/PhysXSupportCore.cpp b/Engine/Source/Runtime/PhysicsCore/Private/PhysXSupportCore.cpp
--- a/Engine/Source/Runtime/PhysicsCore/Private/PhysXSupportCore.cpp
+++ b/Engine/Source/Runtime/PhysicsCore/Private/PhysXSupportCore.cpp
@@ -2,7 +2,7 @@
 
 #include "PhysXSupportCore.h"
 
-#if WITH_PHYSX
+#if PHYSICS_INTERFACE_PHYSX
 #include "PhysicsPublicCore.h"
 #include "PhysXSupportCore.h"
 #include "HAL/IConsoleManager.h"
diff --git a/Engine/Source/Runtime/PhysicsCore/Private/PhysicsInitialization.cpp b/Engine/Source/Runtime/PhysicsCore/Private/PhysicsInitialization.cpp
--- a/Engine/Source/Runtime/PhysicsCore/Private/PhysicsInitialization.cpp
+++ b/Engine/Source/Runtime/PhysicsCore/Private/PhysicsInitialization.cpp
@@ -45,7 +45,7 @@ bool InitGamePhysCore()
 	FModuleManager::Get().LoadModule("ChaosSolverEngine");
 #endif
 
-#if WITH_PHYSX
+#if PHYSICS_INTERFACE_PHYSX
 	// Do nothing if SDK already exists
 	if (GPhysXFoundation != nullptr)
 	{
@@ -190,7 +190,7 @@ bool InitGamePhysCore()
 
 void TermGamePhysCore()
 {
-#if WITH_PHYSX
+#if PHYSICS_INTERFACE_PHYSX
 
 	FPhysxSharedData::Terminate();
 
diff --git a/Engine/Source/Runtime/PhysicsCore/Private/PhysicsInterfaceUtilsCore.cpp b/Engine/Source/Runtime/PhysicsCore/Private/PhysicsInterfaceUtilsCore.cpp
--- a/Engine/Source/Runtime/PhysicsCore/Private/PhysicsInterfaceUtilsCore.cpp
+++ b/Engine/Source/Runtime/PhysicsCore/Private/PhysicsInterfaceUtilsCore.cpp
@@ -21,7 +21,7 @@ PxFilterData U2PFilterData(const FCollisionFilterData& FilterData)
 	return PxFilterData(FilterData.Word0, FilterData.Word1, FilterData.Word2, FilterData.Word3);
 }
 
-#if WITH_PHYSX
+#if PHYSICS_INTERFACE_PHYSX
 
 PxShapeFlags BuildPhysXShapeFlags(FBodyCollisionFlags BodyCollisionFlags, bool bPhysicsStatic, bool bIsTriangleMesh)
 {
diff --git a/Engine/Source/Runtime/PhysicsCore/Private/SQCapture.cpp b/Engine/Source/Runtime/PhysicsCore/Private/SQCapture.cpp
--- a/Engine/Source/Runtime/PhysicsCore/Private/SQCapture.cpp
+++ b/Engine/Source/Runtime/PhysicsCore/Private/SQCapture.cpp
@@ -9,7 +9,7 @@
 #include "PhysicsPublicCore.h"
 #include "PhysTestSerializer.h" 
 
-#if WITH_PHYSX
+#if PHYSICS_INTERFACE_PHYSX
 #include "PhysXToChaosUtil.h"
 #endif
 
@@ -21,7 +21,7 @@ public:
 	virtual ECollisionQueryHitType PostFilter(const FCollisionFilterData& FilterData, const ChaosInterface::FQueryHit& Hit) override { /*check(false);*/  return ECollisionQueryHitType::Touch; }
 	virtual ECollisionQueryHitType PreFilter(const FCollisionFilterData& FilterData, const Chaos::TPerShapeData<float, 3>& Shape, const Chaos::TGeometryParticle<float, 3>& Actor) override { return Capture.GetFilterResult(&Shape, &Actor); }
 
-#if WITH_PHYSX
+#if PHYSICS_INTERFACE_PHYSX
 	virtual ECollisionQueryHitType PostFilter(const FCollisionFilterData& FilterData, const physx::PxQueryHit& Hit) override { return ECollisionQueryHitType::Touch; }
 	virtual ECollisionQueryHitType PreFilter(const FCollisionFilterData& FilterData, const physx::PxShape& Shape, physx::PxRigidActor& Actor) override { return Capture.GetFilterResult(&Shape, &Actor); }
 	virtual PxQueryHitType::Enum preFilter(const PxFilterData& filterData, const PxShape* shape, const PxRigidActor* actor, PxHitFlags& queryFlags) override { return U2PCollisionQueryHitType(Capture.GetFilterResult(shape, actor)); }
@@ -46,7 +46,7 @@ FSQCapture::~FSQCapture()
 {
 }
 
-#if WITH_PHYSX
+#if PHYSICS_INTERFACE_PHYSX
 
 PxActor* FSQCapture::GetTransientActor(PxActor* Actor) const
 {
@@ -295,6 +295,7 @@ void FSQCapture::Serialize(Chaos::FChaosArchive& Ar)
 #if WITH_PHYSX
 	if (bDiskDataIsChaos == false)
 	{
+#if PHYSICS_INTERFACE_PHYSX
 		SerializePhysXBuffers(Ar, Version, PhysXSweepBuffer);
 
 		if (Version >= 1)
@@ -303,16 +304,21 @@ void FSQCapture::Serialize(Chaos::FChaosArchive& Ar)
 			SerializePhysXBuffers(Ar, Version, PhysXOverlapBuffer);
 
 			SerializeActorToShapeHitsArray(Ar);
-			SerializeQueryFilterData(Ar, QueryFilterData);
 		}
 
 		if (Ar.IsLoading())
 		{
 			CreatePhysXData();
 		}
+#endif
+		if(Version >= 1)
+		{
+			SerializeQueryFilterData(Ar, QueryFilterData);
+		}
 	}
 #endif
 
+#if WITH_CHAOS
 	if (bDiskDataIsChaos)
 	{
 #if WITH_CHAOS
@@ -330,10 +336,11 @@ void FSQCapture::Serialize(Chaos::FChaosArchive& Ar)
 			ChaosGeometry = SerializableChaosGeometry.Get();
 		}
 	}
+#endif
 
 	if (Ar.IsLoading())
 	{
-#if WITH_PHYSX
+#if 0
 		CreateChaosDataFromPhysX();
 #endif // WITH_PHYSX
 
@@ -365,6 +372,7 @@ void FSQCapture::StartCaptureChaosSweep(const Chaos::FPBDRigidsEvolutionGBF& Evo
 
 void FSQCapture::EndCaptureChaosSweep(const ChaosInterface::FSQHitBuffer<ChaosInterface::FSweepHit>& Results)
 {
+#if WITH_CHAOS
 	if (IsInGameThread())
 	{
 		check(SQType == ESQType::Sweep);
@@ -372,6 +380,7 @@ void FSQCapture::EndCaptureChaosSweep(const ChaosInterface::FSQHitBuffer<ChaosIn
 		ChaosSweepBuffer = Results;
 #endif // WITH_CHAOS
 	}
+#endif
 }
 
 void FSQCapture::StartCaptureChaosRaycast(const Chaos::FPBDRigidsEvolutionGBF& Evolution, const FVector& InStartPoint, const FVector& InDir, float InDeltaMag, FHitFlags InOutputFlags, const FQueryFilterData& QueryFilter, const FCollisionFilterData& FilterData, ICollisionQueryFilterCallbackBase& Callback)
@@ -393,6 +402,7 @@ void FSQCapture::StartCaptureChaosRaycast(const Chaos::FPBDRigidsEvolutionGBF& E
 
 void FSQCapture::EndCaptureChaosRaycast(const ChaosInterface::FSQHitBuffer<ChaosInterface::FRaycastHit>& Results)
 {
+#if WITH_CHAOS
 	if (IsInGameThread())
 	{
 		check(SQType == ESQType::Raycast);
@@ -400,6 +410,7 @@ void FSQCapture::EndCaptureChaosRaycast(const ChaosInterface::FSQHitBuffer<Chaos
 		ChaosRaycastBuffer = Results;
 #endif // WITH_CHAOS
 	}
+#endif
 }
 
 void FSQCapture::StartCaptureChaosOverlap(const Chaos::FPBDRigidsEvolutionGBF& Evolution, const Chaos::FImplicitObject& InQueryGeom, const FTransform& InStartTM, const FQueryFilterData& QueryFilter, const FCollisionFilterData& FilterData, ICollisionQueryFilterCallbackBase& Callback)
@@ -420,6 +431,7 @@ void FSQCapture::StartCaptureChaosOverlap(const Chaos::FPBDRigidsEvolutionGBF& E
 
 void FSQCapture::EndCaptureChaosOverlap(const ChaosInterface::FSQHitBuffer<ChaosInterface::FOverlapHit>& Results)
 {
+#if WITH_CHAOS
 	if (IsInGameThread())
 	{
 		check(SQType == ESQType::Overlap);
@@ -427,6 +439,7 @@ void FSQCapture::EndCaptureChaosOverlap(const ChaosInterface::FSQHitBuffer<Chaos
 		ChaosOverlapBuffer = Results;
 #endif // WITH_CHAOS
 	}
+#endif
 }
 
 void FSQCapture::CaptureChaosFilterResults(const Chaos::FPBDRigidsEvolutionGBF& TransientEvolution, const FCollisionFilterData& FilterData, ICollisionQueryFilterCallbackBase& Callback)
@@ -451,7 +464,32 @@ void FSQCapture::CaptureChaosFilterResults(const Chaos::FPBDRigidsEvolutionGBF&
 	}
 }
 
-#if WITH_PHYSX
+template <typename TShape, typename TActor>
+ECollisionQueryHitType GetFilterResultHelper(const TShape* Shape, const TActor* Actor, const TMap<TActor*, TArray<TPair<TShape*, ECollisionQueryHitType>>>& ActorToShapeHitsArray)
+{
+	if (const TArray<TPair<TShape*, ECollisionQueryHitType>>* ActorToPairs = ActorToShapeHitsArray.Find(Actor))
+	{
+		for (const TPair<TShape*, ECollisionQueryHitType>& Pair : *ActorToPairs)
+		{
+			if (Pair.Key == Shape)
+			{
+				return Pair.Value;
+			}
+		}
+	}
+
+	//todo: figure out why this hits - suspect it's related to threading and how we capture an evolution on GT
+	ensure(false);	//should not get here, means we didn't properly capture all filter results
+	return ECollisionQueryHitType::None;
+}
+
+ECollisionQueryHitType FSQCapture::GetFilterResult(const Chaos::TPerShapeData<float, 3>* Shape, const Chaos::TGeometryParticle<float, 3>* Actor) const
+{
+	return GetFilterResultHelper(Shape, Actor, ChaosActorToShapeHitsArray);
+}
+
+
+#if PHYSICS_INTERFACE_PHYSX
 void FSQCapture::StartCapturePhysXSweep(const PxScene& Scene, const PxGeometry& InQueryGeom, const FTransform& InStartTM, const FVector& InDir, float InDeltaMag, FHitFlags InOutputFlags, const FQueryFilterData& QueryFilter, const FCollisionFilterData& FilterData, ICollisionQueryFilterCallbackBase& Callback)
 {
 	if (IsInGameThread())
@@ -503,7 +541,6 @@ void FSQCapture::StartCapturePhysXOverlap(const PxScene& Scene, const PxGeometry
 	}
 }
 
-#if WITH_PHYSX
 void FSQCapture::CapturePhysXFilterResults(const PxScene& TransientScene, const FCollisionFilterData& FilterData, ICollisionQueryFilterCallbackBase& Callback)
 {
 	const uint32 NumTransientActors = TransientScene.getNbActors(PxActorTypeFlag::eRIGID_STATIC | PxActorTypeFlag::eRIGID_DYNAMIC);
@@ -538,7 +575,6 @@ void FSQCapture::CapturePhysXFilterResults(const PxScene& TransientScene, const
 	}
 }
 
-#endif // WITH_PHYSX
 
 template <typename THit>
 void EndCaptureHelper(PhysXInterface::FDynamicHitBuffer<THit>& Dest, const PxHitCallback<THit>& Results)
@@ -589,33 +625,6 @@ void FSQCapture::EndCapturePhysXOverlap(const PxHitCallback<PxOverlapHit>& Resul
 		EndCaptureHelper(PhysXOverlapBuffer, Results);
 	}
 }
-#endif
-
-template <typename TShape, typename TActor>
-ECollisionQueryHitType GetFilterResultHelper(const TShape* Shape, const TActor* Actor, const TMap<TActor*, TArray<TPair<TShape*, ECollisionQueryHitType>>>& ActorToShapeHitsArray)
-{
-	if (const TArray<TPair<TShape*, ECollisionQueryHitType>>* ActorToPairs = ActorToShapeHitsArray.Find(Actor))
-	{
-		for (const TPair<TShape*, ECollisionQueryHitType>& Pair : *ActorToPairs)
-		{
-			if (Pair.Key == Shape)
-			{
-				return Pair.Value;
-			}
-		}
-	}
-
-	//todo: figure out why this hits - suspect it's related to threading and how we capture an evolution on GT
-	ensure(false);	//should not get here, means we didn't properly capture all filter results
-	return ECollisionQueryHitType::None;
-}
-
-ECollisionQueryHitType FSQCapture::GetFilterResult(const Chaos::TPerShapeData<float, 3>* Shape, const Chaos::TGeometryParticle<float, 3>* Actor) const
-{
-	return GetFilterResultHelper(Shape, Actor, ChaosActorToShapeHitsArray);
-}
-
-#if WITH_PHYSX
 
 template <typename THit>
 void FixupBufferPointers(FPhysTestSerializer& PhysSerializer, PhysXInterface::FDynamicHitBuffer<THit>& PhysXBuffer)
@@ -732,7 +741,7 @@ FSQCapture::FPhysXSerializerData::~FPhysXSerializerData()
 
 #endif
 
-#if WITH_PHYSX
+#if 0
 void PhysXQueryHitToChaosQueryHit(ChaosInterface::FQueryHit& ChaosHit, const PxQueryHit& PxHit, const FPhysTestSerializer& Serializer)
 {
 	ChaosHit.Actor = Serializer.PhysXActorToChaosHandle(PxHit.actor);
diff --git a/Engine/Source/Runtime/PhysicsCore/Public/CollisionQueryFilterCallbackCore.h b/Engine/Source/Runtime/PhysicsCore/Public/CollisionQueryFilterCallbackCore.h
--- a/Engine/Source/Runtime/PhysicsCore/Public/CollisionQueryFilterCallbackCore.h
+++ b/Engine/Source/Runtime/PhysicsCore/Public/CollisionQueryFilterCallbackCore.h
@@ -4,7 +4,8 @@
 #pragma once
 
 #include "PhysicsInterfaceTypesCore.h"
-#if WITH_PHYSX
+
+#if PHYSICS_INTERFACE_PHYSX
 #include "PhysXPublicCore.h"
 #include "PhysXInterfaceWrapperCore.h"
 #endif
@@ -30,8 +31,8 @@ namespace Chaos
 }
 
 class ICollisionQueryFilterCallbackBase
-#if WITH_PHYSX
-	: public physx::PxQueryFilterCallback	//this is a bit of a hack until physx support is gone. Means any new non physx callback will have dummy pre/postFilter functions for physx api
+#if PHYSICS_INTERFACE_PHYSX
+	: public PxQueryFilterCallback	//this is a bit of a hack until physx support is gone. Means any new non physx callback will have dummy pre/postFilter functions for physx api
 #endif
 {
 public:
@@ -40,7 +41,7 @@ public:
 	virtual ECollisionQueryHitType PostFilter(const FCollisionFilterData& FilterData, const ChaosInterface::FQueryHit& Hit) = 0;
 	virtual ECollisionQueryHitType PreFilter(const FCollisionFilterData& FilterData, const Chaos::TPerShapeData<float, 3>& Shape, const Chaos::TGeometryParticle<float, 3>& Actor) = 0;
 
-#if WITH_PHYSX
+#if  PHYSICS_INTERFACE_PHYSX
 	virtual ECollisionQueryHitType PostFilter(const FCollisionFilterData& FilterData, const physx::PxQueryHit& Hit) = 0;
 	virtual ECollisionQueryHitType PreFilter(const FCollisionFilterData& FilterData, const physx::PxShape& Shape, physx::PxRigidActor& Actor) = 0;
 #endif
@@ -53,7 +54,7 @@ public:
 	virtual ECollisionQueryHitType PostFilter(const FCollisionFilterData& FilterData, const ChaosInterface::FQueryHit& Hit) override { return ECollisionQueryHitType::Block; }
 	virtual ECollisionQueryHitType PreFilter(const FCollisionFilterData& FilterData, const Chaos::TPerShapeData<float, 3>& Shape, const Chaos::TGeometryParticle<float, 3>& Actor) override { return ECollisionQueryHitType::Block; }
 
-#if WITH_PHYSX
+#if PHYSICS_INTERFACE_PHYSX
 	virtual ECollisionQueryHitType PostFilter(const FCollisionFilterData& FilterData, const physx::PxQueryHit& Hit) override { return ECollisionQueryHitType::Block; }
 	virtual ECollisionQueryHitType PreFilter(const FCollisionFilterData& FilterData, const physx::PxShape& Shape, physx::PxRigidActor& Actor) override { return ECollisionQueryHitType::Block; }
 	virtual PxQueryHitType::Enum preFilter(const PxFilterData& filterData, const PxShape* shape, const PxRigidActor* actor, PxHitFlags& queryFlags) override { return PxQueryHitType::eBLOCK; }
@@ -68,7 +69,7 @@ public:
 	virtual ECollisionQueryHitType PostFilter(const FCollisionFilterData& FilterData, const ChaosInterface::FQueryHit& Hit) override { return ECollisionQueryHitType::Touch; }
 	virtual ECollisionQueryHitType PreFilter(const FCollisionFilterData& FilterData, const Chaos::TPerShapeData<float, 3>& Shape, const Chaos::TGeometryParticle<float, 3>& Actor) override { return ECollisionQueryHitType::Touch; }
 
-#if WITH_PHYSX
+#if PHYSICS_INTERFACE_PHYSX
 	virtual ECollisionQueryHitType PostFilter(const FCollisionFilterData& FilterData, const physx::PxQueryHit& Hit) override { return ECollisionQueryHitType::Touch; }
 	virtual ECollisionQueryHitType PreFilter(const FCollisionFilterData& FilterData, const physx::PxShape& Shape, physx::PxRigidActor& Actor) override { return ECollisionQueryHitType::Touch; }
 	virtual PxQueryHitType::Enum preFilter(const PxFilterData& filterData, const PxShape* shape, const PxRigidActor* actor, PxHitFlags& queryFlags) override { return PxQueryHitType::eTOUCH; }
@@ -76,7 +77,7 @@ public:
 #endif
 };
 
-#if WITH_PHYSX
+#if PHYSICS_INTERFACE_PHYSX
 inline PxQueryHitType::Enum U2PCollisionQueryHitType(ECollisionQueryHitType HitType)
 {
 	return (PxQueryHitType::Enum)HitType;
diff --git a/Engine/Source/Runtime/PhysicsCore/Public/PhysTestSerializer.h b/Engine/Source/Runtime/PhysicsCore/Public/PhysTestSerializer.h
--- a/Engine/Source/Runtime/PhysicsCore/Public/PhysTestSerializer.h
+++ b/Engine/Source/Runtime/PhysicsCore/Public/PhysTestSerializer.h
@@ -21,6 +21,7 @@
 
 #if PHYS_TEST_SERIALIZER
 
+#if PHYSICS_INTERFACE_PHYSX
 namespace physx
 {
 	class PxScene;
@@ -30,6 +31,7 @@ namespace physx
 	class PxActor;
 	class PxShape;
 }
+#endif
 
 namespace Chaos
 {
@@ -52,7 +54,10 @@ public:
 
 	//Set the data from an external source. This will obliterate any existing data. Make sure you are not holding on to old internal data as it will go away
 	void SetPhysicsData(Chaos::FPBDRigidsEvolutionGBF& ChaosEvolution);
+
+#if PHYSICS_INTERFACE_PHYSX
 	void SetPhysicsData(physx::PxScene& Scene);
+#endif
 
 	const Chaos::FChaosArchiveContext* GetChaosContext() const
 	{
@@ -71,13 +76,13 @@ public:
 		if (SQCapture)
 		{
 			//todo: this sucks, find a better way to create data instead of doing it lazily
-#if WITH_PHYSX
+#if PHYSICS_INTERFACE_PHYSX
 			GetPhysXData();
 			SQCapture->CreatePhysXData();
 #endif
 
 			GetChaosData();
-#if WITH_PHYSX
+#if 0 
 			SQCapture->CreateChaosDataFromPhysX();
 #endif
 		}
@@ -86,16 +91,18 @@ public:
 
 	Chaos::FPBDRigidsEvolutionGBF* GetChaosData()
 	{
+#if 0
 		if (!bChaosDataReady)
 		{
 			ensure(!bDiskDataIsChaos);
 			//only supported for physx to chaos - don't have serialization context
 			CreateChaosData();
 		}
+#endif
 		return ChaosEvolution.Get();
 	}
 
-#if WITH_PHYSX
+#if PHYSICS_INTERFACE_PHYSX
 	physx::PxScene* GetPhysXData()
 	{
 		if (!bDiskDataIsChaos)	//don't support chaos to physx
@@ -118,8 +125,12 @@ public:
 
 private:
 
+#if 0
 	void CreateChaosData();
+#endif 
+#if PHYSICS_INTERFACE_PHYSX
 	void CreatePhysXData();
+#endif
 
 private:
 
@@ -137,7 +148,7 @@ private:
 
 	FCustomVersionContainer ArchiveVersion;
 
-#if WITH_PHYSX
+#if PHYSICS_INTERFACE_PHYSX
 	struct FPhysXSerializerData
 	{
 		FPhysXSerializerData(int32 NumBytes)
@@ -153,9 +164,7 @@ private:
 		physx::PxSerializationRegistry* Registry;
 	};
 	TUniquePtr<FPhysXSerializerData> AlignedDataHelper;
-#endif
 
-#if WITH_PHYSX
 	TMap<physx::PxActor*, Chaos::TGeometryParticleHandle<float, 3>*> PxActorToChaosHandle;
 	TMap<physx::PxShape*, Chaos::TPerShapeData<float, 3>*> PxShapeToChaosShapes;
 #endif
diff --git a/Engine/Source/Runtime/PhysicsCore/Public/PhysXIncludes.h b/Engine/Source/Runtime/PhysicsCore/Public/PhysXIncludes.h
--- a/Engine/Source/Runtime/PhysicsCore/Public/PhysXIncludes.h
+++ b/Engine/Source/Runtime/PhysicsCore/Public/PhysXIncludes.h
@@ -5,7 +5,7 @@
 
 #include "CoreMinimal.h"
 
-#if WITH_PHYSX
+#if PHYSICS_INTERFACE_PHYSX
 
 #ifdef PHYSX_API
 #undef PHYSX_API	//Our build system treats PhysX as a module and automatically defines this. PhysX has its own way of defining this.
diff --git a/Engine/Source/Runtime/PhysicsCore/Public/PhysXInterfaceWrapperCore.h b/Engine/Source/Runtime/PhysicsCore/Public/PhysXInterfaceWrapperCore.h
--- a/Engine/Source/Runtime/PhysicsCore/Public/PhysXInterfaceWrapperCore.h
+++ b/Engine/Source/Runtime/PhysicsCore/Public/PhysXInterfaceWrapperCore.h
@@ -1,7 +1,7 @@
 // Copyright Epic Games, Inc. All Rights Reserved.
 #pragma once
 
-#if WITH_PHYSX
+#if PHYSICS_INTERFACE_PHYSX
 #include "PhysXPublicCore.h"
 #include "PhysicsInterfaceUtilsCore.h"
 #include "PhysicsInterfaceWrapperShared.h"
diff --git a/Engine/Source/Runtime/PhysicsCore/Public/PhysXPublicCore.h b/Engine/Source/Runtime/PhysicsCore/Public/PhysXPublicCore.h
--- a/Engine/Source/Runtime/PhysicsCore/Public/PhysXPublicCore.h
+++ b/Engine/Source/Runtime/PhysicsCore/Public/PhysXPublicCore.h
@@ -9,7 +9,7 @@
 #include "CoreMinimal.h"
 #include "Stats/Stats.h"
 
-#if WITH_PHYSX
+#if PHYSICS_INTERFACE_PHYSX
 #include "PhysXIncludes.h"
 #include "PhysicsInterfaceTypesCore.h"
 #include "CollisionShape.h"
diff --git a/Engine/Source/Runtime/PhysicsCore/Public/PhysXSupportCore.h b/Engine/Source/Runtime/PhysicsCore/Public/PhysXSupportCore.h
--- a/Engine/Source/Runtime/PhysicsCore/Public/PhysXSupportCore.h
+++ b/Engine/Source/Runtime/PhysicsCore/Public/PhysXSupportCore.h
@@ -10,7 +10,7 @@
 #include "PhysicsCore.h"
 #include "PhysicsPublicCore.h"
 
-#if WITH_PHYSX
+#if PHYSICS_INTERFACE_PHYSX
 #include "PhysXPublicCore.h"
 #include "HAL/IConsoleManager.h"
 
diff --git a/Engine/Source/Runtime/PhysicsCore/Public/PhysXToChaosUtil.h b/Engine/Source/Runtime/PhysicsCore/Public/PhysXToChaosUtil.h
--- a/Engine/Source/Runtime/PhysicsCore/Public/PhysXToChaosUtil.h
+++ b/Engine/Source/Runtime/PhysicsCore/Public/PhysXToChaosUtil.h
@@ -2,6 +2,7 @@
 
 #pragma once 
 
+#if 0
 #include "CoreMinimal.h"
 
 #if WITH_PHYSX
@@ -196,4 +197,6 @@ inline TUniquePtr<Chaos::TImplicitObjectTransformed<float, 3>> PxShapeToChaosGeo
 	return MakeUnique<TImplicitObjectTransformed<float, 3>>(MoveTemp(InnerObj), P2UTransform(ShapeTM));
 }
 
+#endif
+
 #endif
\ No newline at end of file
diff --git a/Engine/Source/Runtime/PhysicsCore/Public/PhysicsInterfaceTypesCore.h b/Engine/Source/Runtime/PhysicsCore/Public/PhysicsInterfaceTypesCore.h
--- a/Engine/Source/Runtime/PhysicsCore/Public/PhysicsInterfaceTypesCore.h
+++ b/Engine/Source/Runtime/PhysicsCore/Public/PhysicsInterfaceTypesCore.h
@@ -144,10 +144,78 @@ static bool CalcMeshNegScaleCompensation(const FVector& InScale3D, FTransform& O
 }
 
 
-#if !WITH_PHYSX
+//TODO: Reimplement types in chaos
+#if !PHYSICS_INTERFACE_PHYSX
 
 const uint32 AggregateMaxSize = 128;
 
+
+class UPhysicalMaterial;
+class UPrimitiveComponent;
+struct FBodyInstance;
+struct FConstraintInstance;
+struct FKShapeElem;
+
+/** Forward declarations */
+struct FKShapeElem;
+struct FCustomPhysXPayload;
+
+/** PhysX user data type*/
+namespace EPhysxUserDataType
+{
+	enum Type
+	{
+		Invalid,
+		BodyInstance,
+		PhysicalMaterial,
+		PhysScene,
+		ConstraintInstance,
+		PrimitiveComponent,
+		AggShape,
+		CustomPayload,	//This is intended for plugins
+	};
+};
+
+/** PhysX user data */
+struct FPhysxUserData
+{
+protected:
+	EPhysxUserDataType::Type	Type;
+	void*						Payload;
+
+public:
+	FPhysxUserData()									:Type(EPhysxUserDataType::Invalid), Payload(nullptr) {}
+	FPhysxUserData(FBodyInstance* InPayload)			:Type(EPhysxUserDataType::BodyInstance), Payload(InPayload) {}
+	FPhysxUserData(UPhysicalMaterial* InPayload)		:Type(EPhysxUserDataType::PhysicalMaterial), Payload(InPayload) {}
+	FPhysxUserData(FPhysScene* InPayload)			    :Type(EPhysxUserDataType::PhysScene), Payload(InPayload) {}
+	FPhysxUserData(FConstraintInstance* InPayload)		:Type(EPhysxUserDataType::ConstraintInstance), Payload(InPayload) {}
+	FPhysxUserData(UPrimitiveComponent* InPayload)		:Type(EPhysxUserDataType::PrimitiveComponent), Payload(InPayload) {}
+	FPhysxUserData(FKShapeElem* InPayload)				:Type(EPhysxUserDataType::AggShape), Payload(InPayload) {}
+	FPhysxUserData(FCustomPhysXPayload* InPayload)		:Type(EPhysxUserDataType::CustomPayload), Payload(InPayload) {}
+	
+	template <class T> static T* Get(void* UserData);
+	template <class T> static void Set(void* UserData, T* Payload);
+
+	//helper function to determine if userData is garbage (maybe dangling pointer)
+	static bool IsGarbage(void* UserData){ return ((FPhysxUserData*)UserData)->Type < EPhysxUserDataType::Invalid || ((FPhysxUserData*)UserData)->Type > EPhysxUserDataType::CustomPayload; }
+};
+
+template <> FORCEINLINE FBodyInstance* FPhysxUserData::Get(void* UserData)			{ if (!UserData || ((FPhysxUserData*)UserData)->Type != EPhysxUserDataType::BodyInstance) { return nullptr; } return (FBodyInstance*)((FPhysxUserData*)UserData)->Payload; }
+template <> FORCEINLINE UPhysicalMaterial* FPhysxUserData::Get(void* UserData)		{ if (!UserData || ((FPhysxUserData*)UserData)->Type != EPhysxUserDataType::PhysicalMaterial) { return nullptr; } return (UPhysicalMaterial*)((FPhysxUserData*)UserData)->Payload; }
+template <> FORCEINLINE FPhysScene* FPhysxUserData::Get(void* UserData)				{ if (!UserData || ((FPhysxUserData*)UserData)->Type != EPhysxUserDataType::PhysScene) { return nullptr; }return (FPhysScene*)((FPhysxUserData*)UserData)->Payload; }
+template <> FORCEINLINE FConstraintInstance* FPhysxUserData::Get(void* UserData)	{ if (!UserData || ((FPhysxUserData*)UserData)->Type != EPhysxUserDataType::ConstraintInstance) { return nullptr; } return (FConstraintInstance*)((FPhysxUserData*)UserData)->Payload; }
+template <> FORCEINLINE UPrimitiveComponent* FPhysxUserData::Get(void* UserData)	{ if (!UserData || ((FPhysxUserData*)UserData)->Type != EPhysxUserDataType::PrimitiveComponent) { return nullptr; } return (UPrimitiveComponent*)((FPhysxUserData*)UserData)->Payload; }
+template <> FORCEINLINE FKShapeElem* FPhysxUserData::Get(void* UserData)	{ if (!UserData || ((FPhysxUserData*)UserData)->Type != EPhysxUserDataType::AggShape) { return nullptr; } return (FKShapeElem*)((FPhysxUserData*)UserData)->Payload; }
+template <> FORCEINLINE FCustomPhysXPayload* FPhysxUserData::Get(void* UserData) { if (!UserData || ((FPhysxUserData*)UserData)->Type != EPhysxUserDataType::CustomPayload) { return nullptr; } return (FCustomPhysXPayload*)((FPhysxUserData*)UserData)->Payload; }
+
+template <> FORCEINLINE void FPhysxUserData::Set(void* UserData, FBodyInstance* Payload)			{ check(UserData); ((FPhysxUserData*)UserData)->Type = EPhysxUserDataType::BodyInstance; ((FPhysxUserData*)UserData)->Payload = Payload; }
+template <> FORCEINLINE void FPhysxUserData::Set(void* UserData, UPhysicalMaterial* Payload)		{ check(UserData); ((FPhysxUserData*)UserData)->Type = EPhysxUserDataType::PhysicalMaterial; ((FPhysxUserData*)UserData)->Payload = Payload; }
+template <> FORCEINLINE void FPhysxUserData::Set(void* UserData, FPhysScene* Payload)				{ check(UserData); ((FPhysxUserData*)UserData)->Type = EPhysxUserDataType::PhysScene; ((FPhysxUserData*)UserData)->Payload = Payload; }
+template <> FORCEINLINE void FPhysxUserData::Set(void* UserData, FConstraintInstance* Payload)		{ check(UserData); ((FPhysxUserData*)UserData)->Type = EPhysxUserDataType::ConstraintInstance; ((FPhysxUserData*)UserData)->Payload = Payload; }
+template <> FORCEINLINE void FPhysxUserData::Set(void* UserData, UPrimitiveComponent* Payload)		{ check(UserData); ((FPhysxUserData*)UserData)->Type = EPhysxUserDataType::PrimitiveComponent; ((FPhysxUserData*)UserData)->Payload = Payload; }
+template <> FORCEINLINE void FPhysxUserData::Set(void* UserData, FKShapeElem* Payload)	{ check(UserData); ((FPhysxUserData*)UserData)->Type = EPhysxUserDataType::AggShape; ((FPhysxUserData*)UserData)->Payload = Payload; }
+template <> FORCEINLINE void FPhysxUserData::Set(void* UserData, FCustomPhysXPayload* Payload) { check(UserData); ((FPhysxUserData*)UserData)->Type = EPhysxUserDataType::CustomPayload; ((FPhysxUserData*)UserData)->Payload = Payload; }
+
 /** enum for empty constructor tag*/
 enum PxEMPTY
 {
diff --git a/Engine/Source/Runtime/PhysicsCore/Public/PhysicsInterfaceUtilsCore.h b/Engine/Source/Runtime/PhysicsCore/Public/PhysicsInterfaceUtilsCore.h
--- a/Engine/Source/Runtime/PhysicsCore/Public/PhysicsInterfaceUtilsCore.h
+++ b/Engine/Source/Runtime/PhysicsCore/Public/PhysicsInterfaceUtilsCore.h
@@ -5,7 +5,7 @@
 #include "PhysicsInterfaceTypesCore.h"
 #include "PhysicsInterfaceDeclaresCore.h"
 
-#if WITH_PHYSX
+#if PHYSICS_INTERFACE_PHYSX
 #include "PhysXPublicCore.h"
 #include "Containers/Union.h"
 #include "CollisionShape.h"
@@ -14,7 +14,7 @@
 PHYSICSCORE_API FCollisionFilterData P2UFilterData(const PxFilterData& PFilterData);
 PHYSICSCORE_API PxFilterData U2PFilterData(const FCollisionFilterData& FilterData);
 
-#if WITH_PHYSX
+#if PHYSICS_INTERFACE_PHYSX
 
 PHYSICSCORE_API PxShapeFlags BuildPhysXShapeFlags(FBodyCollisionFlags BodyCollisionFlags, bool bPhysicsStatic, bool bIsTriangleMesh);
 PHYSICSCORE_API PxGeometryType::Enum U2PCollisionShapeType(ECollisionShapeType InUType);
diff --git a/Engine/Source/Runtime/PhysicsCore/Public/SQCapture.h b/Engine/Source/Runtime/PhysicsCore/Public/SQCapture.h
--- a/Engine/Source/Runtime/PhysicsCore/Public/SQCapture.h
+++ b/Engine/Source/Runtime/PhysicsCore/Public/SQCapture.h
@@ -43,7 +43,7 @@ struct PHYSICSCORE_API FSQCapture
 		Overlap
 	} SQType;
 
-#if WITH_PHYSX
+#if PHYSICS_INTERFACE_PHYSX
 	void StartCapturePhysXSweep(const PxScene& Scene, const PxGeometry& InQueryGeom, const FTransform& InStartTM, const FVector& InDir, float InDeltaMag, FHitFlags InOutputFlags, const FQueryFilterData& QueryFilter, const FCollisionFilterData& FilterData, ICollisionQueryFilterCallbackBase& Callback);
 	void EndCapturePhysXSweep(const PxHitCallback<PxSweepHit>& Results);
 
@@ -65,7 +65,7 @@ struct PHYSICSCORE_API FSQCapture
 
 	ECollisionQueryHitType GetFilterResult(const Chaos::TPerShapeData<float,3>* Shape, const Chaos::TGeometryParticle<float,3>* Actor) const;
 
-#if WITH_PHYSX
+#if PHYSICS_INTERFACE_PHYSX
 	ECollisionQueryHitType GetFilterResult(const physx::PxShape* Shape, const physx::PxActor* Actor) const;
 #endif
 	
@@ -78,7 +78,7 @@ struct PHYSICSCORE_API FSQCapture
 	FQueryFilterData QueryFilterData;
 	TUniquePtr<ICollisionQueryFilterCallbackBase> FilterCallback;
 
-#if WITH_PHYSX
+#if PHYSICS_INTERFACE_PHYSX
 	PhysXInterface::FDynamicHitBuffer<PxSweepHit> PhysXSweepBuffer;
 	PhysXInterface::FDynamicHitBuffer<PxRaycastHit> PhysXRaycastBuffer;
 	PhysXInterface::FDynamicHitBuffer<PxOverlapHit> PhysXOverlapBuffer;
@@ -104,7 +104,7 @@ private:
 	FSQCapture(FPhysTestSerializer& OwningPhysSerializer);	//This should be created by PhysTestSerializer
 	void Serialize(Chaos::FChaosArchive& Ar);
 
-#if WITH_PHYSX
+#if PHYSICS_INTERFACE_PHYSX
 	void SerializeActorToShapeHitsArray(FArchive& Ar);
 #endif
 
@@ -113,12 +113,12 @@ private:
 	TArray<uint8> GeomData;
 	TArray<uint8> HitData;
 
-#if WITH_PHYSX
+#if 0
 	void CreateChaosDataFromPhysX();
 	void CreateChaosFilterResults();
 #endif
 
-#if WITH_PHYSX
+#if PHYSICS_INTERFACE_PHYSX
 
 	static constexpr uint64 ShapeCollectionID = 1;
 
diff --git a/Engine/Source/Runtime/PhysicsSQ/Private/SQAccelerator.cpp b/Engine/Source/Runtime/PhysicsSQ/Private/SQAccelerator.cpp
--- a/Engine/Source/Runtime/PhysicsSQ/Private/SQAccelerator.cpp
+++ b/Engine/Source/Runtime/PhysicsSQ/Private/SQAccelerator.cpp
@@ -2,6 +2,7 @@
 
 #include "SQAccelerator.h"
 #include "CollisionQueryFilterCallbackCore.h"
+#include "PhysicsCore/Public/PhysicsInterfaceUtilsCore.h"
 
 #if PHYSICS_INTERFACE_PHYSX
 #include "SceneQueryPhysXImp.h"	//todo: use nice platform wrapper
@@ -12,7 +13,6 @@
 #include "Experimental/SceneQueryChaosImp.h"
 #endif
 
-#include "PhysicsInterfaceUtilsCore.h"
 #include "ChaosInterfaceWrapperCore.h"
 #include "Chaos/CastingUtilities.h"
 #include "Chaos/ISpatialAcceleration.h"
@@ -98,8 +98,10 @@ struct TSQVisitor : public Chaos::ISpatialVisitor<TPayload, float>
 		, QueryFilterDataConcrete(P2UFilterData(QueryFilterData.data))
 		, QueryCallback(InQueryCallback)
 	{
+#if WITH_PHYSX
 		//#TODO - reimplement query flags alternative for Chaos
 		bAnyHit = QueryFilterData.flags & PxQueryFlag::eANY_HIT;
+#endif
 	}
 
 	TSQVisitor(const FTransform& InStartTM, const FVector& InDir, ChaosInterface::FSQHitBuffer<ChaosInterface::FSweepHit>& InHitBuffer, EHitFlags InOutputFlags,
diff --git a/Engine/Source/Runtime/PhysicsSQ/Public/SQVerifier.h b/Engine/Source/Runtime/PhysicsSQ/Public/SQVerifier.h
--- a/Engine/Source/Runtime/PhysicsSQ/Public/SQVerifier.h
+++ b/Engine/Source/Runtime/PhysicsSQ/Public/SQVerifier.h
@@ -7,11 +7,127 @@
 #if WITH_PHYSX
 
 #include "Chaos/PBDRigidsEvolutionGBF.h"
+#include "PhysicsCore/Public/PhysicsInterfaceUtilsCore.h"
 
 #ifndef SQ_REPLAY_TEST
 #define SQ_REPLAY_TEST(cond) bEnsureOnMismatch ? ensure(cond) : (cond)
 #endif
 
+template <bool bHasPhysX = false>
+void SQPerfComparisonHelper(const FString& TestName, FPhysTestSerializer& Serializer, bool bEnsureOnMismatch = false)
+{
+	using namespace Chaos;
+	uint32 PhysXSum = 0;
+	uint32 ChaosSum = 0;
+	double NumIterations = bHasPhysX ? 100 : 10000;
+	//double NumIterations = 1000000;
+
+	const FSQCapture& CapturedSQ = *Serializer.GetSQCapture();
+	switch (CapturedSQ.SQType)
+	{
+	case FSQCapture::ESQType::Raycast:
+	{
+#if PHYSICS_INTERFACE_PHYSX
+		if (bHasPhysX)
+		{
+			for (double i = 0; i < NumIterations; ++i)
+			{
+				auto PxHitBuffer = MakeUnique<PhysXInterface::FDynamicHitBuffer<PxRaycastHit>>();
+				uint32 StartTime = FPlatformTime::Cycles();
+				Serializer.GetPhysXData()->raycast(U2PVector(CapturedSQ.StartPoint), U2PVector(CapturedSQ.Dir), CapturedSQ.DeltaMag, *PxHitBuffer, U2PHitFlags(CapturedSQ.OutputFlags.HitFlags), CapturedSQ.QueryFilterData, CapturedSQ.FilterCallback.Get());
+				PhysXSum += FPlatformTime::Cycles() - StartTime;
+			}
+		}
+#endif
+
+		TUniquePtr<ISpatialAccelerationCollection<TAccelerationStructureHandle<float, 3>, float, 3>> Accelerator;
+		Serializer.GetChaosData()->UpdateExternalAccelerationStructure(Accelerator);
+		FChaosSQAccelerator SQAccelerator(*Accelerator);
+		for (double i = 0; i < NumIterations; ++i)
+		{
+			auto ChaosHitBuffer = MakeUnique<ChaosInterface::FSQHitBuffer<ChaosInterface::FRaycastHit>>();
+			uint32 StartTime = FPlatformTime::Cycles();
+			SQAccelerator.Raycast(CapturedSQ.StartPoint, CapturedSQ.Dir, CapturedSQ.DeltaMag, *ChaosHitBuffer, CapturedSQ.OutputFlags.HitFlags, CapturedSQ.QueryFilterData, *CapturedSQ.FilterCallback);
+			ChaosSum += FPlatformTime::Cycles() - StartTime;
+		}
+		break;
+	}
+	case FSQCapture::ESQType::Sweep:
+	{
+#if PHYSICS_INTERFACE_PHYSX
+		if (bHasPhysX)
+		{
+			for (double i = 0; i < NumIterations; ++i)
+			{
+				auto PxHitBuffer = MakeUnique<PhysXInterface::FDynamicHitBuffer<PxSweepHit>>();
+				uint32 StartTime = FPlatformTime::Cycles();
+				Serializer.GetPhysXData()->sweep(CapturedSQ.PhysXGeometry.any(), U2PTransform(CapturedSQ.StartTM), U2PVector(CapturedSQ.Dir), CapturedSQ.DeltaMag, *PxHitBuffer, U2PHitFlags(CapturedSQ.OutputFlags.HitFlags), CapturedSQ.QueryFilterData, CapturedSQ.FilterCallback.Get());
+				PhysXSum += FPlatformTime::Cycles() - StartTime;
+			}
+		}
+#endif
+
+		TUniquePtr<ISpatialAccelerationCollection<TAccelerationStructureHandle<float, 3>, float, 3>> Accelerator;
+		Serializer.GetChaosData()->UpdateExternalAccelerationStructure(Accelerator);
+		FChaosSQAccelerator SQAccelerator(*Accelerator);
+		for (double i = 0; i < NumIterations; ++i)
+		{
+			auto ChaosHitBuffer = MakeUnique<ChaosInterface::FSQHitBuffer<ChaosInterface::FSweepHit>>();
+			uint32 StartTime = FPlatformTime::Cycles();
+			SQAccelerator.Sweep(*CapturedSQ.ChaosGeometry, CapturedSQ.StartTM, CapturedSQ.Dir, CapturedSQ.DeltaMag, *ChaosHitBuffer, CapturedSQ.OutputFlags.HitFlags, CapturedSQ.QueryFilterData, *CapturedSQ.FilterCallback);
+			ChaosSum += FPlatformTime::Cycles() - StartTime;
+		}
+		break;
+	}
+	case FSQCapture::ESQType::Overlap:
+	{
+#if PHYSICS_INTERFACE_PHYSX
+		if (bHasPhysX)
+		{
+			for (double i = 0; i < NumIterations; ++i)
+			{
+				auto PxHitBuffer = MakeUnique<PhysXInterface::FDynamicHitBuffer<PxOverlapHit>>();
+				uint32 StartTime = FPlatformTime::Cycles();
+				Serializer.GetPhysXData()->overlap(CapturedSQ.PhysXGeometry.any(), U2PTransform(CapturedSQ.StartTM), *PxHitBuffer, CapturedSQ.QueryFilterData, CapturedSQ.FilterCallback.Get());
+				PhysXSum += FPlatformTime::Cycles() - StartTime;
+			}
+		}
+#endif
+
+		TUniquePtr<ISpatialAccelerationCollection<TAccelerationStructureHandle<float, 3>, float, 3>> Accelerator;
+		Serializer.GetChaosData()->UpdateExternalAccelerationStructure(Accelerator);
+		FChaosSQAccelerator SQAccelerator(*Accelerator);
+		for (double i = 0; i < NumIterations; ++i)
+		{
+			auto ChaosHitBuffer = MakeUnique<ChaosInterface::FSQHitBuffer<ChaosInterface::FOverlapHit>>();
+			uint32 StartTime = FPlatformTime::Cycles();
+			SQAccelerator.Overlap(*CapturedSQ.ChaosGeometry, CapturedSQ.StartTM, *ChaosHitBuffer, CapturedSQ.QueryFilterData, *CapturedSQ.FilterCallback);
+			ChaosSum += FPlatformTime::Cycles() - StartTime;
+		}
+		break;
+	}
+	}
+
+	float ChaosMs = FPlatformTime::ToMilliseconds(ChaosSum);
+	float ChaosAvgUs = (ChaosMs * 1000) / NumIterations;
+
+	float PhysXMs = FPlatformTime::ToMilliseconds(PhysXSum);
+	float PhysXAvgUs = (PhysXMs * 1000) / NumIterations;
+
+
+	if (bHasPhysX)
+	{
+		UE_LOG(LogPhysicsCore, Warning, TEXT("Perf Test:%s\nPhysX:%f(us), Chaos:%f(us)"), *TestName, PhysXAvgUs, ChaosAvgUs);
+	}
+	else
+	{
+		UE_LOG(LogPhysicsCore, Warning, TEXT("Perf Test:%s\nChaos:%f(us), Total:%f(ms)"), *TestName, ChaosAvgUs, ChaosMs);
+		//UE_LOG(LogPhysicsCore, Warning, TEXT("Perf Test:%s\nChaos:%f(us)"), *TestName, AvgChaos);
+	}
+}
+
+
+#if 0
 bool SQComparisonHelper(FPhysTestSerializer& Serializer, bool bEnsureOnMismatch = false)
 {
 	using namespace Chaos;
@@ -128,112 +244,6 @@
 	return bTestPassed;
 }
 
-template <bool bHasPhysX = true>
-void SQPerfComparisonHelper(const FString& TestName, FPhysTestSerializer& Serializer, bool bEnsureOnMismatch = false)
-{
-	using namespace Chaos;
-	uint32 PhysXSum = 0;
-	uint32 ChaosSum = 0;
-	double NumIterations = bHasPhysX ? 100 : 10000;
-	//double NumIterations = 1000000;
-
-	const FSQCapture& CapturedSQ = *Serializer.GetSQCapture();
-	switch (CapturedSQ.SQType)
-	{
-	case FSQCapture::ESQType::Raycast:
-	{
-		if (bHasPhysX)
-		{
-			for (double i = 0; i < NumIterations; ++i)
-			{
-				auto PxHitBuffer = MakeUnique<PhysXInterface::FDynamicHitBuffer<PxRaycastHit>>();
-				uint32 StartTime = FPlatformTime::Cycles();
-				Serializer.GetPhysXData()->raycast(U2PVector(CapturedSQ.StartPoint), U2PVector(CapturedSQ.Dir), CapturedSQ.DeltaMag, *PxHitBuffer, U2PHitFlags(CapturedSQ.OutputFlags.HitFlags), CapturedSQ.QueryFilterData, CapturedSQ.FilterCallback.Get());
-				PhysXSum += FPlatformTime::Cycles() - StartTime;
-			}
-		}
-
-		TUniquePtr<ISpatialAccelerationCollection<TAccelerationStructureHandle<float, 3>, float, 3>> Accelerator;
-		Serializer.GetChaosData()->UpdateExternalAccelerationStructure(Accelerator);
-		FChaosSQAccelerator SQAccelerator(*Accelerator);
-		for (double i = 0; i < NumIterations; ++i)
-		{
-			auto ChaosHitBuffer = MakeUnique<ChaosInterface::FSQHitBuffer<ChaosInterface::FRaycastHit>>();
-			uint32 StartTime = FPlatformTime::Cycles();
-			SQAccelerator.Raycast(CapturedSQ.StartPoint, CapturedSQ.Dir, CapturedSQ.DeltaMag, *ChaosHitBuffer, CapturedSQ.OutputFlags.HitFlags, CapturedSQ.QueryFilterData, *CapturedSQ.FilterCallback);
-			ChaosSum += FPlatformTime::Cycles() - StartTime;
-		}
-		break;
-	}
-	case FSQCapture::ESQType::Sweep:
-	{
-		if (bHasPhysX)
-		{
-			for (double i = 0; i < NumIterations; ++i)
-			{
-				auto PxHitBuffer = MakeUnique<PhysXInterface::FDynamicHitBuffer<PxSweepHit>>();
-				uint32 StartTime = FPlatformTime::Cycles();
-				Serializer.GetPhysXData()->sweep(CapturedSQ.PhysXGeometry.any(), U2PTransform(CapturedSQ.StartTM), U2PVector(CapturedSQ.Dir), CapturedSQ.DeltaMag, *PxHitBuffer, U2PHitFlags(CapturedSQ.OutputFlags.HitFlags), CapturedSQ.QueryFilterData, CapturedSQ.FilterCallback.Get());
-				PhysXSum += FPlatformTime::Cycles() - StartTime;
-			}
-		}
-
-		TUniquePtr<ISpatialAccelerationCollection<TAccelerationStructureHandle<float, 3>, float, 3>> Accelerator;
-		Serializer.GetChaosData()->UpdateExternalAccelerationStructure(Accelerator);
-		FChaosSQAccelerator SQAccelerator(*Accelerator);
-		for (double i = 0; i < NumIterations; ++i)
-		{
-			auto ChaosHitBuffer = MakeUnique<ChaosInterface::FSQHitBuffer<ChaosInterface::FSweepHit>>();
-			uint32 StartTime = FPlatformTime::Cycles();
-			SQAccelerator.Sweep(*CapturedSQ.ChaosGeometry, CapturedSQ.StartTM, CapturedSQ.Dir, CapturedSQ.DeltaMag, *ChaosHitBuffer, CapturedSQ.OutputFlags.HitFlags, CapturedSQ.QueryFilterData, *CapturedSQ.FilterCallback);
-			ChaosSum += FPlatformTime::Cycles() - StartTime;
-		}
-		break;
-	}
-	case FSQCapture::ESQType::Overlap:
-	{
-		if (bHasPhysX)
-		{
-			for (double i = 0; i < NumIterations; ++i)
-			{
-				auto PxHitBuffer = MakeUnique<PhysXInterface::FDynamicHitBuffer<PxOverlapHit>>();
-				uint32 StartTime = FPlatformTime::Cycles();
-				Serializer.GetPhysXData()->overlap(CapturedSQ.PhysXGeometry.any(), U2PTransform(CapturedSQ.StartTM), *PxHitBuffer, CapturedSQ.QueryFilterData, CapturedSQ.FilterCallback.Get());
-				PhysXSum += FPlatformTime::Cycles() - StartTime;
-			}
-		}
-
-		TUniquePtr<ISpatialAccelerationCollection<TAccelerationStructureHandle<float, 3>, float, 3>> Accelerator;
-		Serializer.GetChaosData()->UpdateExternalAccelerationStructure(Accelerator);
-		FChaosSQAccelerator SQAccelerator(*Accelerator);
-		for (double i = 0; i < NumIterations; ++i)
-		{
-			auto ChaosHitBuffer = MakeUnique<ChaosInterface::FSQHitBuffer<ChaosInterface::FOverlapHit>>();
-			uint32 StartTime = FPlatformTime::Cycles();
-			SQAccelerator.Overlap(*CapturedSQ.ChaosGeometry, CapturedSQ.StartTM, *ChaosHitBuffer, CapturedSQ.QueryFilterData, *CapturedSQ.FilterCallback);
-			ChaosSum += FPlatformTime::Cycles() - StartTime;
-		}
-		break;
-	}
-	}
-
-	float ChaosMs = FPlatformTime::ToMilliseconds(ChaosSum);
-	float ChaosAvgUs = (ChaosMs * 1000) / NumIterations;
-
-	float PhysXMs = FPlatformTime::ToMilliseconds(PhysXSum);
-	float PhysXAvgUs = (PhysXMs * 1000) / NumIterations;
-
-	if (bHasPhysX)
-	{
-		UE_LOG(LogPhysicsCore, Warning, TEXT("Perf Test:%s\nPhysX:%f(us), Chaos:%f(us)"), *TestName, PhysXAvgUs, ChaosAvgUs);
-	}
-	else
-	{
-		UE_LOG(LogPhysicsCore, Warning, TEXT("Perf Test:%s\nChaos:%f(us), Total:%f(ms)"), *TestName, ChaosAvgUs, ChaosMs);
-		//UE_LOG(LogPhysicsCore, Warning, TEXT("Perf Test:%s\nChaos:%f(us)"), *TestName, AvgChaos);
-	}
-}
-
 #endif
 
 #if INCLUDE_CHAOS
@@ -303,3 +313,4 @@
 }
 
 #endif
+#endif
