SQAccelerator payload debug drawing

diff --git a/Engine/Source/Runtime/Experimental/Chaos/Private/Chaos/ChaosDebugDraw.cpp b/Engine/Source/Runtime/Experimental/Chaos/Private/Chaos/ChaosDebugDraw.cpp
--- a/Engine/Source/Runtime/Experimental/Chaos/Private/Chaos/ChaosDebugDraw.cpp
+++ b/Engine/Source/Runtime/Experimental/Chaos/Private/Chaos/ChaosDebugDraw.cpp
@@ -125,6 +125,14 @@ namespace Chaos
 			DrawShapesImpl(FRigidTransform3(P, Q), Particle->Geometry().Get(), Color);
 		}
 
+		void DrawParticleShapesImpl(const FRigidTransform3& SpaceTransform, const TGeometryParticle<FReal, 3>* Particle, const FColor& Color)
+		{
+			FVec3 P = SpaceTransform.TransformPosition(Particle->ObjectState() == EObjectStateType::Dynamic ? Particle->CastToRigidParticle()->P() : Particle->X());
+			FRotation3 Q = SpaceTransform.GetRotation() * (Particle->ObjectState() == EObjectStateType::Dynamic ? Particle->CastToRigidParticle()->Q() : Particle->R());
+
+			DrawShapesImpl(FRigidTransform3(P, Q), Particle->Geometry().Get(), Color);
+		}
+
 		void DrawParticleBoundsImpl(const FRigidTransform3& SpaceTransform, const TGeometryParticleHandle<FReal, 3>* InParticle, const FColor& Color)
 		{
 			TAABB<FReal, 3> Box = InParticle->WorldSpaceInflatedBounds();
@@ -365,6 +373,26 @@ namespace Chaos
 #endif
 		}
 
+		void DrawParticleShapes(const FRigidTransform3& SpaceTransform, const TGeometryParticleHandle<float, 3>* Particle, const FColor& Color)
+		{
+#if CHAOS_DEBUG_DRAW
+			if (FDebugDrawQueue::IsDebugDrawingEnabled())
+			{
+				DrawParticleShapesImpl(SpaceTransform, Particle, Color);
+			}
+#endif
+		}
+
+		void DrawParticleShapes(const FRigidTransform3& SpaceTransform, const TGeometryParticle<float, 3>* Particle, const FColor& Color)
+		{
+#if CHAOS_DEBUG_DRAW
+			if (FDebugDrawQueue::IsDebugDrawingEnabled())
+			{
+				DrawParticleShapesImpl(SpaceTransform, Particle, Color);
+			}
+#endif
+		}
+
 		void DrawParticleBounds(const FRigidTransform3& SpaceTransform, const TParticleView<TGeometryParticles<float, 3>>& ParticlesView, const FColor& Color)
 		{
 #if CHAOS_DEBUG_DRAW
diff --git a/Engine/Source/Runtime/Experimental/Chaos/Public/Chaos/ChaosDebugDraw.h b/Engine/Source/Runtime/Experimental/Chaos/Public/Chaos/ChaosDebugDraw.h
--- a/Engine/Source/Runtime/Experimental/Chaos/Public/Chaos/ChaosDebugDraw.h
+++ b/Engine/Source/Runtime/Experimental/Chaos/Public/Chaos/ChaosDebugDraw.h
@@ -29,6 +29,8 @@ namespace Chaos
 		CHAOS_API void DrawParticleShapes(const FRigidTransform3& SpaceTransform, const TParticleView<TGeometryParticles<float, 3>>& ParticlesView, const FColor& Color);
 		CHAOS_API void DrawParticleShapes(const FRigidTransform3& SpaceTransform, const TParticleView<TKinematicGeometryParticles<float, 3>>& ParticlesView, const FColor& Color);
 		CHAOS_API void DrawParticleShapes(const FRigidTransform3& SpaceTransform, const TParticleView<TPBDRigidParticles<float, 3>>& ParticlesView, const FColor& Color);
+		CHAOS_API void DrawParticleShapes(const FRigidTransform3& SpaceTransform, const TGeometryParticleHandle<float, 3>* Particle, const FColor& Color);
+		CHAOS_API void DrawParticleShapes(const FRigidTransform3& SpaceTransform, const TGeometryParticle<float, 3>* Particle, const FColor& Color);
 		CHAOS_API void DrawParticleBounds(const FRigidTransform3& SpaceTransform, const TParticleView<TGeometryParticles<float, 3>>& ParticlesView, const FColor& Color);
 		CHAOS_API void DrawParticleBounds(const FRigidTransform3& SpaceTransform, const TParticleView<TKinematicGeometryParticles<float, 3>>& ParticlesView, const FColor& Color);
 		CHAOS_API void DrawParticleBounds(const FRigidTransform3& SpaceTransform, const TParticleView<TPBDRigidParticles<float, 3>>& ParticlesView, const FColor& Color);
diff --git a/Engine/Source/Runtime/Experimental/Chaos/Public/Chaos/ParticleHandle.h b/Engine/Source/Runtime/Experimental/Chaos/Public/Chaos/ParticleHandle.h
--- a/Engine/Source/Runtime/Experimental/Chaos/Public/Chaos/ParticleHandle.h
+++ b/Engine/Source/Runtime/Experimental/Chaos/Public/Chaos/ParticleHandle.h
@@ -8,6 +8,10 @@
 #include "Chaos/ParticleIterator.h"
 #include "Chaos/ParticleDirtyFlags.h"
 #include "ChaosCheck.h"
+#include "Chaos/ChaosDebugDrawDeclares.h"
+#if CHAOS_DEBUG_DRAW
+#include "Chaos/ChaosDebugDraw.h"
+#endif
 
 
 
@@ -225,6 +229,11 @@ private:
 
 	template <typename TParticle>
 	void UpdatePrePreFilter(const TParticle& Particle);
+
+public:
+#if CHAOS_DEBUG_DRAW
+	void DebugDraw(const bool bExternal) const;
+#endif
 };
 
 template <typename T, int d>
@@ -2472,6 +2481,23 @@ FChaosArchive& operator<<(FChaosArchive& Ar, TAccelerationStructureHandle<T, d>&
 	AccelerationHandle.Serialize(Ar);
 	return Ar;
 }
+
+#if CHAOS_DEBUG_DRAW
+template <typename T, int d>
+void TAccelerationStructureHandle<T, d>::DebugDraw(const bool bExternal) const
+{
+	if (ExternalGeometryParticle && bExternal)
+	{
+		DebugDraw::DrawParticleShapes(FRigidTransform3(), ExternalGeometryParticle, FColor(0, 255, 0));
+	}
+
+	if (GeometryParticleHandle && !bExternal)
+	{
+		DebugDraw::DrawParticleShapes(FRigidTransform3(), GeometryParticleHandle, FColor(100, 255, 100));
+	}
+}
+#endif
+
 #if PLATFORM_MAC || PLATFORM_LINUX
 extern template class CHAOS_API TGeometryParticle<float, 3>;
 extern template class CHAOS_API TKinematicGeometryParticle<float, 3>;
diff --git a/Engine/Source/Runtime/PhysicsCore/Public/ChaosInterfaceWrapperCore.h b/Engine/Source/Runtime/PhysicsCore/Public/ChaosInterfaceWrapperCore.h
--- a/Engine/Source/Runtime/PhysicsCore/Public/ChaosInterfaceWrapperCore.h
+++ b/Engine/Source/Runtime/PhysicsCore/Public/ChaosInterfaceWrapperCore.h
@@ -81,13 +81,17 @@ struct FQueryDebugParams
 {
 #if !(UE_BUILD_TEST || UE_BUILD_SHIPPING) 
 	FQueryDebugParams()
-		: bDebugQuery(false) { }
+		: bDebugQuery(false)
+		, bExternalQuery(true) { }
 	bool bDebugQuery;
+	bool bExternalQuery;
 	bool IsDebugQuery() const { return bDebugQuery; }
+	bool IsExternalQuery() const { return bExternalQuery; }
 #else
 	// In test or shipping builds, this struct must be left empty
 	FQueryDebugParams() { }
 	constexpr bool IsDebugQuery() const { return false; }
+	constexpr bool IsExternalQuery() const { return true; }
 #endif
 };
 #endif
diff --git a/Engine/Source/Runtime/PhysicsSQ/Private/SQAccelerator.cpp b/Engine/Source/Runtime/PhysicsSQ/Private/SQAccelerator.cpp
--- a/Engine/Source/Runtime/PhysicsSQ/Private/SQAccelerator.cpp
+++ b/Engine/Source/Runtime/PhysicsSQ/Private/SQAccelerator.cpp
@@ -194,13 +194,6 @@ private:
 	bool Visit(const Chaos::TSpatialVisitorData<TPayload>& Instance, Chaos::FQueryFastData* CurData)
 	{
 		//QUICK_SCOPE_CYCLE_COUNTER(SQVisit);
-#if CHAOS_DEBUG_DRAW && WITH_CHAOS
-		if (DebugParams.IsDebugQuery() && ChaosSQDrawDebugVisitorQueries)
-		{
-			DebugDraw<SQ>(Instance, CurData->CurrentLength);
-		}
-#endif
-
 		TPayload Payload = Instance.Payload;
 
 		//todo: add a check to ensure hitbuffer matches SQ type
@@ -212,6 +205,10 @@ private:
 
 		const TRigidTransform<float, 3> ActorTM(GeometryParticle->X(), GeometryParticle->R());
 
+#if CHAOS_DEBUG_DRAW
+		bool bAllShapesIgnoredInPrefilter = true;
+#endif
+
 		for (const auto& Shape : Shapes)
 		{
 			const FImplicitObject* Geom = Shape->Geometry.Get();
@@ -259,6 +256,10 @@ private:
 
 			if (HitType != ECollisionQueryHitType::None)
 			{
+#if CHAOS_DEBUG_DRAW
+				bAllShapesIgnoredInPrefilter = false;
+#endif
+
 				//QUICK_SCOPE_CYCLE_COUNTER(SQNarrow);
 				THitType Hit;
 				Hit.Actor = GeometryParticle;
@@ -343,22 +344,39 @@ private:
 			}
 		}
 
+#if CHAOS_DEBUG_DRAW && WITH_CHAOS
+		if (DebugParams.IsDebugQuery() && ChaosSQDrawDebugVisitorQueries)
+		{
+			DebugDraw<SQ>(Instance, DebugParams, CurData->CurrentLength, bAllShapesIgnoredInPrefilter);
+		}
+#endif
+
 		return true;
 	}
 
-#if CHAOS_DEBUG_DRAW && !(UE_BUILD_TEST || UE_BUILD_SHIPPING)
+#if CHAOS_DEBUG_DRAW
+
+	template <typename TPayload> void DebugDrawPayloadImpl(const TPayload& Payload, const bool bExternal, decltype(&TPayload::DebugDraw)) { Payload.DebugDraw(bExternal); }
+	template <typename TPayload> void DebugDrawPayloadImpl(const TPayload& Payload, const bool bExternal, ...) { }
+	template <typename TPayload> void DebugDrawPayload(const TPayload& Payload, const bool bExternal) { DebugDrawPayloadImpl(Payload, bExternal, 0); }
+
 	template <ESQType SQ>
-	void DebugDraw(const Chaos::TSpatialVisitorData<TPayload>& Instance, const float CurLength)
+	void DebugDraw(const Chaos::TSpatialVisitorData<TPayload>& Instance, const FQueryDebugParams& DebugParams, const float CurLength, const bool bPrefiltered)
 	{
 		if (SQ == ESQType::Raycast)
 		{
 			const FVector EndPoint = StartPoint + (Dir * CurLength);
-			Chaos::FDebugDrawQueue::GetInstance().DrawDebugDirectionalArrow(StartPoint, EndPoint, 5.f, FColor::Green, false, -1.f, 0, 1.f);
+			Chaos::FDebugDrawQueue::GetInstance().DrawDebugDirectionalArrow(StartPoint, EndPoint, 5.f, FColor::Green);
 		}
 
 		if (Instance.bHasBounds)
 		{
-			Chaos::FDebugDrawQueue::GetInstance().DrawDebugBox(Instance.Bounds.Center(), Instance.Bounds.Extents(), FQuat::Identity, FColor::Red, false, -1.f, 0, 2.f);
+			Chaos::FDebugDrawQueue::GetInstance().DrawDebugBox(Instance.Bounds.Center(), Instance.Bounds.Extents(), FQuat::Identity, FColor(50, 200, 50), false, -1.f, 0, 0.f);
+		}
+
+		if (!bPrefiltered)
+		{
+			DebugDrawPayload(Instance.Payload, DebugParams.bExternalQuery);
 		}
 	}
 #endif
