Fix Linux compile error. Add SQ debug hit indication.

diff --git a/Engine/Source/Runtime/Experimental/Chaos/Public/Chaos/ParticleHandle.h b/Engine/Source/Runtime/Experimental/Chaos/Public/Chaos/ParticleHandle.h
--- a/Engine/Source/Runtime/Experimental/Chaos/Public/Chaos/ParticleHandle.h
+++ b/Engine/Source/Runtime/Experimental/Chaos/Public/Chaos/ParticleHandle.h
@@ -232,7 +232,7 @@ private:
 
 public:
 #if CHAOS_DEBUG_DRAW
-	void DebugDraw(const bool bExternal) const;
+	void DebugDraw(const bool bExternal, const bool bHit) const;
 #endif
 };
 
@@ -2484,16 +2484,16 @@ FChaosArchive& operator<<(FChaosArchive& Ar, TAccelerationStructureHandle<T, d>&
 
 #if CHAOS_DEBUG_DRAW
 template <typename T, int d>
-void TAccelerationStructureHandle<T, d>::DebugDraw(const bool bExternal) const
+void TAccelerationStructureHandle<T, d>::DebugDraw(const bool bExternal, const bool bHit) const
 {
 	if (ExternalGeometryParticle && bExternal)
 	{
-		DebugDraw::DrawParticleShapes(FRigidTransform3(), ExternalGeometryParticle, FColor(0, 255, 0));
+		DebugDraw::DrawParticleShapes(FRigidTransform3(), ExternalGeometryParticle, bHit ? FColor::Red : FColor::Green);
 	}
 
 	if (GeometryParticleHandle && !bExternal)
 	{
-		DebugDraw::DrawParticleShapes(FRigidTransform3(), GeometryParticleHandle, FColor(100, 255, 100));
+		DebugDraw::DrawParticleShapes(FRigidTransform3(), GeometryParticleHandle, bHit ? FColor(200, 100, 100) : FColor(100, 200, 100));
 	}
 }
 #endif
diff --git a/Engine/Source/Runtime/PhysicsSQ/Private/SQAccelerator.cpp b/Engine/Source/Runtime/PhysicsSQ/Private/SQAccelerator.cpp
--- a/Engine/Source/Runtime/PhysicsSQ/Private/SQAccelerator.cpp
+++ b/Engine/Source/Runtime/PhysicsSQ/Private/SQAccelerator.cpp
@@ -202,11 +202,13 @@ private:
 		const TShapesArray<float,3>& Shapes = GeometryParticle->ShapesArray();
 
 		const bool bTestShapeBounds =  Shapes.Num() > 1;
+		bool bContinue = true;
 
 		const TRigidTransform<float, 3> ActorTM(GeometryParticle->X(), GeometryParticle->R());
 
 #if CHAOS_DEBUG_DRAW
 		bool bAllShapesIgnoredInPrefilter = true;
+		bool bHitBufferIncreased = false;
 #endif
 
 		for (const auto& Shape : Shapes)
@@ -325,19 +327,24 @@ private:
 						//overlap never blocks
 						const bool bBlocker = (HitType == ECollisionQueryHitType::Block || bAnyHit || HitBuffer.WantsSingleResult());
 						HitBuffer.InsertHit(Hit, bBlocker);
+#if CHAOS_DEBUG_DRAW
+						bHitBufferIncreased = true;
+#endif
 
 						if (bBlocker && SQ != ESQType::Overlap)
 						{
 							CurData->SetLength(FMath::Max(0.f, Distance));	//Max is needed for MTD which returns negative distance
 							if (CurData->CurrentLength == 0 && (SQ == ESQType::Raycast || HitBuffer.WantsSingleResult()))	//raycasts always fail with distance 0, sweeps only matter if we want multi overlaps
 							{
-								return false;	//initial overlap so nothing will be better than this
+								bContinue = false; //initial overlap so nothing will be better than this
+								break;
 							}
 						}
 
 						if (bAnyHit)
 						{
-							return false;
+							bContinue = false;
+							break;
 						}
 					}
 				}
@@ -347,36 +354,36 @@ private:
 #if CHAOS_DEBUG_DRAW && WITH_CHAOS
 		if (DebugParams.IsDebugQuery() && ChaosSQDrawDebugVisitorQueries)
 		{
-			DebugDraw<SQ>(Instance, DebugParams, CurData->CurrentLength, bAllShapesIgnoredInPrefilter);
+			DebugDraw<SQ>(Instance, DebugParams, CurData->CurrentLength, bAllShapesIgnoredInPrefilter, bHitBufferIncreased);
 		}
 #endif
 
-		return true;
+		return bContinue;
 	}
 
 #if CHAOS_DEBUG_DRAW
 
-	template <typename TPayload> void DebugDrawPayloadImpl(const TPayload& Payload, const bool bExternal, decltype(&TPayload::DebugDraw)) { Payload.DebugDraw(bExternal); }
-	template <typename TPayload> void DebugDrawPayloadImpl(const TPayload& Payload, const bool bExternal, ...) { }
-	template <typename TPayload> void DebugDrawPayload(const TPayload& Payload, const bool bExternal) { DebugDrawPayloadImpl(Payload, bExternal, 0); }
+	void DebugDrawPayloadImpl(const TPayload& Payload, const bool bExternal, const bool bHit, decltype(&TPayload::DebugDraw)) { Payload.DebugDraw(bExternal, bHit); }
+	void DebugDrawPayloadImpl(const TPayload& Payload, const bool bExternal, const bool bHit, ...) { }
+	void DebugDrawPayload(const TPayload& Payload, const bool bExternal, const bool bHit) { DebugDrawPayloadImpl(Payload, bExternal, bHit, 0); }
 
 	template <ESQType SQ>
-	void DebugDraw(const Chaos::TSpatialVisitorData<TPayload>& Instance, const FQueryDebugParams& DebugParams, const float CurLength, const bool bPrefiltered)
+	void DebugDraw(const Chaos::TSpatialVisitorData<TPayload>& Instance, const FQueryDebugParams& DebugParams, const float CurLength, const bool bPrefiltered, const bool bHit)
 	{
 		if (SQ == ESQType::Raycast)
 		{
 			const FVector EndPoint = StartPoint + (Dir * CurLength);
-			Chaos::FDebugDrawQueue::GetInstance().DrawDebugDirectionalArrow(StartPoint, EndPoint, 5.f, FColor::Green);
+			Chaos::FDebugDrawQueue::GetInstance().DrawDebugDirectionalArrow(StartPoint, EndPoint, 5.f, bHit ? FColor::Red : FColor::Green);
 		}
 
 		if (Instance.bHasBounds)
 		{
-			Chaos::FDebugDrawQueue::GetInstance().DrawDebugBox(Instance.Bounds.Center(), Instance.Bounds.Extents(), FQuat::Identity, FColor(50, 200, 50), false, -1.f, 0, 0.f);
+			Chaos::FDebugDrawQueue::GetInstance().DrawDebugBox(Instance.Bounds.Center(), Instance.Bounds.Extents(), FQuat::Identity, bHit ? FColor(100, 50, 50) : FColor(50, 100, 50), false, -1.f, 0, 0.f);
 		}
 
 		if (!bPrefiltered)
 		{
-			DebugDrawPayload(Instance.Payload, DebugParams.bExternalQuery);
+			DebugDrawPayload(Instance.Payload, DebugParams.bExternalQuery, bHit);
 		}
 	}
 #endif
