Remove unnecessary DebugParams argument to DebugDraw.

diff --git a/Engine/Source/Runtime/PhysicsSQ/Private/SQAccelerator.cpp b/Engine/Source/Runtime/PhysicsSQ/Private/SQAccelerator.cpp
--- a/Engine/Source/Runtime/PhysicsSQ/Private/SQAccelerator.cpp
+++ b/Engine/Source/Runtime/PhysicsSQ/Private/SQAccelerator.cpp
@@ -354,7 +354,7 @@ private:
 #if CHAOS_DEBUG_DRAW && WITH_CHAOS
 		if (DebugParams.IsDebugQuery() && ChaosSQDrawDebugVisitorQueries)
 		{
-			DebugDraw<SQ>(Instance, DebugParams, CurData->CurrentLength, bAllShapesIgnoredInPrefilter, bHitBufferIncreased);
+			DebugDraw<SQ>(Instance, CurData->CurrentLength, bAllShapesIgnoredInPrefilter, bHitBufferIncreased);
 		}
 #endif
 
@@ -368,7 +368,7 @@ private:
 	void DebugDrawPayload(const TPayload& Payload, const bool bExternal, const bool bHit) { DebugDrawPayloadImpl(Payload, bExternal, bHit, 0); }
 
 	template <ESQType SQ>
-	void DebugDraw(const Chaos::TSpatialVisitorData<TPayload>& Instance, const FQueryDebugParams& DebugParams, const float CurLength, const bool bPrefiltered, const bool bHit)
+	void DebugDraw(const Chaos::TSpatialVisitorData<TPayload>& Instance, const float CurLength, const bool bPrefiltered, const bool bHit)
 	{
 		if (SQ == ESQType::Raycast)
 		{
