Fixing up compile issue if Chaos is off in SQAccelerator.cpp
#rnx

diff --git a/Engine/Source/Runtime/PhysicsSQ/Private/SQAccelerator.cpp b/Engine/Source/Runtime/PhysicsSQ/Private/SQAccelerator.cpp
--- a/Engine/Source/Runtime/PhysicsSQ/Private/SQAccelerator.cpp
+++ b/Engine/Source/Runtime/PhysicsSQ/Private/SQAccelerator.cpp
@@ -381,10 +381,12 @@ private:
 			Chaos::FDebugDrawQueue::GetInstance().DrawDebugBox(Instance.Bounds.Center(), Instance.Bounds.Extents(), FQuat::Identity, bHit ? FColor(100, 50, 50) : FColor(50, 100, 50), false, -1.f, 0, 0.f);
 		}
 
+#if WITH_CHAOS
 		if (!bPrefiltered)
 		{
 			DebugDrawPayload(Instance.Payload, DebugParams.bExternalQuery, bHit);
 		}
+#endif
 	}
 #endif
 
