Add option to skip narrow phase for overlap queries.

diff --git a/Engine/Source/Runtime/Engine/Private/Collision/SceneQuery.cpp b/Engine/Source/Runtime/Engine/Private/Collision/SceneQuery.cpp
--- a/Engine/Source/Runtime/Engine/Private/Collision/SceneQuery.cpp
+++ b/Engine/Source/Runtime/Engine/Private/Collision/SceneQuery.cpp
@@ -558,8 +558,11 @@ bool GeomOverlapMultiImp(const UWorld* World, const FPhysicsGeometry& Geom, cons
 		QueryCallback.bIgnoreTouches |= (InfoType == EQueryInfo::IsBlocking); // pre-filter to ignore touches and only get blocking hits, if that's what we're after.
 		QueryCallback.bIsOverlapQuery = true;
 
-		const EQueryFlags QueryFlags = InfoType == EQueryInfo::GatherAll ? EQueryFlags::PreFilter : (EQueryFlags::PreFilter | EQueryFlags::AnyHit);
-		
+		EQueryFlags QueryFlags = InfoType == EQueryInfo::GatherAll ? EQueryFlags::PreFilter : (EQueryFlags::PreFilter | EQueryFlags::AnyHit);
+		if (Params.bSkipNarrowPhase)
+		{
+			QueryFlags = QueryFlags | EQueryFlags::SkipNarrowPhase;
+		}
 		FDynamicHitBuffer<FHitOverlap> OverlapBuffer;
 
 		// Enable scene locks, in case they are required
diff --git a/Engine/Source/Runtime/Engine/Public/CollisionQueryParams.h b/Engine/Source/Runtime/Engine/Public/CollisionQueryParams.h
--- a/Engine/Source/Runtime/Engine/Public/CollisionQueryParams.h
+++ b/Engine/Source/Runtime/Engine/Public/CollisionQueryParams.h
@@ -63,6 +63,9 @@ struct ENGINE_API FCollisionQueryParams
 	/** Whether to ignore touch/overlap results. */
 	bool bIgnoreTouches;
 
+	/** Whether to skip narrow phase checks (only for overlaps). */
+	bool bSkipNarrowPhase;
+
 	/** Filters query by mobility types (static vs stationary/movable)*/
 	EQueryMobilityType MobilityType;
 
@@ -152,6 +155,7 @@ public:
 		IgnoreMask = 0;
 		bIgnoreBlocks = false;
 		bIgnoreTouches = false;
+		bSkipNarrowPhase = false;
 		StatId = GetUnknownStatId();
 #if !(UE_BUILD_TEST || UE_BUILD_SHIPPING)
 		bDebugQuery = false;
@@ -170,6 +174,7 @@ public:
 		IgnoreMask = 0;
 		bIgnoreBlocks = false;
 		bIgnoreTouches = false;
+		bSkipNarrowPhase = false;
 		StatId = GetUnknownStatId();
 #if !(UE_BUILD_TEST || UE_BUILD_SHIPPING)
 		bDebugQuery = false;
diff --git a/Engine/Source/Runtime/PhysicsCore/Public/PhysicsInterfaceTypesCore.h b/Engine/Source/Runtime/PhysicsCore/Public/PhysicsInterfaceTypesCore.h
--- a/Engine/Source/Runtime/PhysicsCore/Public/PhysicsInterfaceTypesCore.h
+++ b/Engine/Source/Runtime/PhysicsCore/Public/PhysicsInterfaceTypesCore.h
@@ -271,6 +271,8 @@ struct FChaosQueryFlag
 		eNO_BLOCK = (1 << 5),	//!< All hits are reported as touching. Overrides eBLOCK returned from user filters with eTOUCH.
 								//!< This is also an optimization hint that may improve query performance.
 
+		eSKIPNARROWPHASE = (1 << 6), //!< Skip narrow phase check for the query
+
 		eRESERVED = (1 << 15)	//!< Reserved for internal use
 	};
 };
@@ -353,6 +355,11 @@ inline FChaosQueryFlags U2CQueryFlags(FQueryFlags Flags)
 		Result |= FChaosQueryFlag::eANY_HIT;
 	}
 
+	if (Flags & EQueryFlags::SkipNarrowPhase)
+	{
+		Result |= FChaosQueryFlag::eSKIPNARROWPHASE;
+	}
+	
 	return (FChaosQueryFlags)Result;
 }
 
diff --git a/Engine/Source/Runtime/PhysicsCore/Public/PhysicsInterfaceWrapperShared.h b/Engine/Source/Runtime/PhysicsCore/Public/PhysicsInterfaceWrapperShared.h
--- a/Engine/Source/Runtime/PhysicsCore/Public/PhysicsInterfaceWrapperShared.h
+++ b/Engine/Source/Runtime/PhysicsCore/Public/PhysicsInterfaceWrapperShared.h
@@ -7,7 +7,8 @@ enum class EQueryFlags : uint16
 	None = 0,
 	PreFilter = (1 << 2),
 	PostFilter = (1 << 3),
-	AnyHit = (1 << 4)
+	AnyHit = (1 << 4),
+	SkipNarrowPhase = (1 << 5)
 };
 
 inline EQueryFlags operator| (EQueryFlags lhs, EQueryFlags rhs)
diff --git a/Engine/Source/Runtime/PhysicsSQ/Private/SQAccelerator.cpp b/Engine/Source/Runtime/PhysicsSQ/Private/SQAccelerator.cpp
--- a/Engine/Source/Runtime/PhysicsSQ/Private/SQAccelerator.cpp
+++ b/Engine/Source/Runtime/PhysicsSQ/Private/SQAccelerator.cpp
@@ -412,6 +412,103 @@ private:
 	const FTransform StartTM;
 };
 
+template <typename QueryGeometryType, typename TPayload, typename THitType>
+struct TBPVisitor : public Chaos::ISpatialVisitor<TPayload, float>
+{
+	TBPVisitor(const FTransform& InWorldTM, ChaosInterface::FSQHitBuffer<ChaosInterface::FOverlapHit>& InHitBuffer,
+		const FQueryFilterData& InQueryFilterData, ICollisionQueryFilterCallbackBase& InQueryCallback, const QueryGeometryType& InQueryGeom, const FQueryDebugParams& InDebugParams)
+		: HalfExtents(InQueryGeom.BoundingBox().Extents() * 0.5)
+		, bAnyHit(false)
+		, DebugParams(InDebugParams)
+		, HitBuffer(InHitBuffer)
+		, QueryFilterData(InQueryFilterData)
+		, QueryFilterDataConcrete(C2UFilterData(QueryFilterData.data))
+		, QueryGeom(&InQueryGeom)
+		, QueryCallback(InQueryCallback)
+		, StartTM(InWorldTM)
+	{
+		bAnyHit = QueryFilterData.flags & FChaosQueryFlag::eANY_HIT;
+	}
+	virtual bool Overlap(const Chaos::TSpatialVisitorData<TPayload>& Instance) override
+	{
+		return Visit<ESQType::Overlap>(Instance, nullptr);
+	}
+
+	virtual bool Raycast(const Chaos::TSpatialVisitorData<TPayload>& Instance, Chaos::FQueryFastData& CurData) override
+	{
+		ensure(false);
+		return false;
+	}
+
+	virtual bool Sweep(const Chaos::TSpatialVisitorData<TPayload>& Instance, Chaos::FQueryFastData& CurData) override
+	{
+		ensure(false);
+		return false;
+	}
+
+	virtual const void* GetQueryData() const override
+	{
+		return &QueryFilterData;
+	}
+
+private:
+
+	enum class ESQType
+	{
+		Raycast,
+		Sweep,
+		Overlap
+	};
+
+	template <ESQType SQ>
+	bool Visit(const Chaos::TSpatialVisitorData<TPayload>& Instance, Chaos::FQueryFastData* CurData)
+	{
+		bool bContinue = true;
+		TPayload Payload = Instance.Payload;
+		//todo: add a check to ensure hitbuffer matches SQ type
+		using namespace Chaos;
+		TGeometryParticle<float, 3>* GeometryParticle = Payload.GetExternalGeometryParticle_ExternalThread();
+		if (!GeometryParticle)
+		{
+			// This case handles particles created by the physics simulation without the main thread
+			// being made aware of their creation. We have a PT particle but no external particle
+			return true;
+		}
+		const FShapesArray& Shapes = GeometryParticle->ShapesArray();
+		THitType Hit;
+		Hit.Actor = GeometryParticle;
+		for (const auto& Shape : Shapes)
+		{
+			ECollisionQueryHitType HitType = QueryFilterData.flags & FChaosQueryFlag::ePREFILTER ? QueryCallback.PreFilter(QueryFilterDataConcrete, *Shape, *GeometryParticle) : ECollisionQueryHitType::Block;
+			if (HitType != ECollisionQueryHitType::None)
+			{
+				const bool bBlocker = (HitType == ECollisionQueryHitType::Block || bAnyHit || HitBuffer.WantsSingleResult());
+				Hit.Shape = Shape.Get();
+				HitBuffer.InsertHit(Hit, bBlocker);
+				if (bAnyHit)
+				{
+					bContinue = false;
+				}
+				break;
+			}
+		}
+		return bContinue;
+	}
+
+	const FVector StartPoint;
+	const FVector Dir;
+	const FVector HalfExtents;
+	FHitFlags OutputFlags;
+	bool bAnyHit;
+	const FQueryDebugParams DebugParams;
+	ChaosInterface::FSQHitBuffer<THitType>& HitBuffer;
+	const FQueryFilterData& QueryFilterData;
+	const FCollisionFilterData QueryFilterDataConcrete;
+	const QueryGeometryType* QueryGeom;
+	ICollisionQueryFilterCallbackBase& QueryCallback;
+	const FTransform StartTM;
+};
+
 void FChaosSQAccelerator::Raycast(const FVector& Start, const FVector& Dir, const float DeltaMagnitude, ChaosInterface::FSQHitBuffer<ChaosInterface::FRaycastHit>& HitBuffer, EHitFlags OutputFlags, const FQueryFilterData& QueryFilterData, ICollisionQueryFilterCallbackBase& QueryCallback, const FQueryDebugParams& DebugParams) const
 {
 	using namespace Chaos;
@@ -460,9 +557,18 @@ void OverlapHelper(const QueryGeomType& QueryGeom, const Chaos::ISpatialAccelera
 	using namespace ChaosInterface;
 
 	const TAABB<float, 3> Bounds = QueryGeom.BoundingBox().TransformedAABB(GeomPose);
-	TSQVisitor<QueryGeomType, TAccelerationStructureHandle<float, 3>, FOverlapHit> OverlapVisitor(GeomPose, HitBuffer, QueryFilterData, QueryCallback, QueryGeom, DebugParams);
+
 	HitBuffer.IncFlushCount();
-	SpatialAcceleration.Overlap(Bounds, OverlapVisitor);
+	if (QueryFilterData.flags & FChaosQueryFlag::eSKIPNARROWPHASE)
+	{
+		TBPVisitor<QueryGeomType, TAccelerationStructureHandle<float, 3>, FOverlapHit> OverlapVisitor(GeomPose, HitBuffer, QueryFilterData, QueryCallback, QueryGeom, DebugParams);
+		SpatialAcceleration.Overlap(Bounds, OverlapVisitor);
+	}
+	else
+	{
+		TSQVisitor<QueryGeomType, TAccelerationStructureHandle<float, 3>, FOverlapHit> OverlapVisitor(GeomPose, HitBuffer, QueryFilterData, QueryCallback, QueryGeom, DebugParams);
+		SpatialAcceleration.Overlap(Bounds, OverlapVisitor);
+	}
 	HitBuffer.DecFlushCount();
 }
 
