Fix PhysX Build

diff --git a/Engine/Source/Runtime/PhysicsCore/Private/PhysicsInterfaceUtilsCore.cpp b/Engine/Source/Runtime/PhysicsCore/Private/PhysicsInterfaceUtilsCore.cpp
--- a/Engine/Source/Runtime/PhysicsCore/Private/PhysicsInterfaceUtilsCore.cpp
+++ b/Engine/Source/Runtime/PhysicsCore/Private/PhysicsInterfaceUtilsCore.cpp
@@ -22,6 +22,19 @@ FChaosFilterData U2CFilterData(const FCollisionFilterData& FilterData)
 	return FChaosFilterData(FilterData.Word0, FilterData.Word1, FilterData.Word2, FilterData.Word3);
 }
 
+#if PHYSICS_INTERFACE_PHYSX
+FCollisionFilterData ToUnrealFilterData(const PxFilterData& FilterData)
+{
+	return P2UFilterData(FilterData);
+}
+#else
+FCollisionFilterData ToUnrealFilterData(const FChaosFilterData& FilterData)
+{
+	return C2UFilterData(FilterData);
+}
+#endif
+
+
 #if PHYSICS_INTERFACE_PHYSX
 
 FCollisionFilterData P2UFilterData(const PxFilterData& PFilterData)
diff --git a/Engine/Source/Runtime/PhysicsCore/Public/PhysicsInterfaceDeclaresCore.h b/Engine/Source/Runtime/PhysicsCore/Public/PhysicsInterfaceDeclaresCore.h
--- a/Engine/Source/Runtime/PhysicsCore/Public/PhysicsInterfaceDeclaresCore.h
+++ b/Engine/Source/Runtime/PhysicsCore/Public/PhysicsInterfaceDeclaresCore.h
@@ -90,6 +90,7 @@ namespace physx
 	struct PxSweepHit;
 	struct PxRaycastHit;
 	struct PxOverlapHit;
+	struct PxQueryFlag;
 	struct PxQueryHit;
 
 	class PxTransform;
diff --git a/Engine/Source/Runtime/PhysicsCore/Public/PhysicsInterfaceTypesCore.h b/Engine/Source/Runtime/PhysicsCore/Public/PhysicsInterfaceTypesCore.h
--- a/Engine/Source/Runtime/PhysicsCore/Public/PhysicsInterfaceTypesCore.h
+++ b/Engine/Source/Runtime/PhysicsCore/Public/PhysicsInterfaceTypesCore.h
@@ -7,6 +7,7 @@
 #include "Chaos/CollisionFilterData.h"
 
 struct FBodyInstance;
+struct FChaosQueryFlag;
 
 struct FActorCreationParams
 {
@@ -202,8 +203,10 @@ public:
 
 #if PHYSICS_INTERFACE_PHYSX
 using FUserData = FPhysxUserData;
+using FPhysicsQueryFlag = physx::PxQueryFlag;
 #else
 using FUserData = FChaosUserData;
+using FPhysicsQueryFlag = FChaosQueryFlag;
 #endif
 
 
diff --git a/Engine/Source/Runtime/PhysicsCore/Public/PhysicsInterfaceUtilsCore.h b/Engine/Source/Runtime/PhysicsCore/Public/PhysicsInterfaceUtilsCore.h
--- a/Engine/Source/Runtime/PhysicsCore/Public/PhysicsInterfaceUtilsCore.h
+++ b/Engine/Source/Runtime/PhysicsCore/Public/PhysicsInterfaceUtilsCore.h
@@ -15,7 +15,12 @@ PHYSICSCORE_API FCollisionFilterData C2UFilterData(const FChaosFilterData& Filte
 PHYSICSCORE_API FChaosFilterData U2CFilterData(const FCollisionFilterData& FilterData);
 
 #if PHYSICS_INTERFACE_PHYSX
+PHYSICSCORE_API FCollisionFilterData ToUnrealFilterData(const PxFilterData& FilterData);
+#else
+PHYSICSCORE_API FCollisionFilterData ToUnrealFilterData(const FChaosFilterData& FilterData);
+#endif
 
+#if PHYSICS_INTERFACE_PHYSX
 PHYSICSCORE_API FCollisionFilterData P2UFilterData(const PxFilterData& PFilterData);
 PHYSICSCORE_API PxFilterData U2PFilterData(const FCollisionFilterData& FilterData);
 PHYSICSCORE_API PxShapeFlags BuildPhysXShapeFlags(FBodyCollisionFlags BodyCollisionFlags, bool bPhysicsStatic, bool bIsTriangleMesh);
diff --git a/Engine/Source/Runtime/PhysicsSQ/Private/SQAccelerator.cpp b/Engine/Source/Runtime/PhysicsSQ/Private/SQAccelerator.cpp
--- a/Engine/Source/Runtime/PhysicsSQ/Private/SQAccelerator.cpp
+++ b/Engine/Source/Runtime/PhysicsSQ/Private/SQAccelerator.cpp
@@ -422,12 +422,12 @@ struct TBPVisitor : public Chaos::ISpatialVisitor<TPayload, float>
 		, DebugParams(InDebugParams)
 		, HitBuffer(InHitBuffer)
 		, QueryFilterData(InQueryFilterData)
-		, QueryFilterDataConcrete(C2UFilterData(QueryFilterData.data))
+		, QueryFilterDataConcrete(ToUnrealFilterData(InQueryFilterData.data))
 		, QueryGeom(&InQueryGeom)
 		, QueryCallback(InQueryCallback)
 		, StartTM(InWorldTM)
 	{
-		bAnyHit = QueryFilterData.flags & FChaosQueryFlag::eANY_HIT;
+		bAnyHit = QueryFilterData.flags & FPhysicsQueryFlag::eANY_HIT;
 	}
 	virtual bool Overlap(const Chaos::TSpatialVisitorData<TPayload>& Instance) override
 	{
@@ -479,7 +479,7 @@ private:
 		Hit.Actor = GeometryParticle;
 		for (const auto& Shape : Shapes)
 		{
-			ECollisionQueryHitType HitType = QueryFilterData.flags & FChaosQueryFlag::ePREFILTER ? QueryCallback.PreFilter(QueryFilterDataConcrete, *Shape, *GeometryParticle) : ECollisionQueryHitType::Block;
+			ECollisionQueryHitType HitType = QueryFilterData.flags & FPhysicsQueryFlag::ePREFILTER ? QueryCallback.PreFilter(QueryFilterDataConcrete, *Shape, *GeometryParticle) : ECollisionQueryHitType::Block;
 			if (HitType != ECollisionQueryHitType::None)
 			{
 				const bool bBlocker = (HitType == ECollisionQueryHitType::Block || bAnyHit || HitBuffer.WantsSingleResult());
@@ -559,7 +559,12 @@ void OverlapHelper(const QueryGeomType& QueryGeom, const Chaos::ISpatialAccelera
 	const TAABB<float, 3> Bounds = QueryGeom.BoundingBox().TransformedAABB(GeomPose);
 
 	HitBuffer.IncFlushCount();
-	if (QueryFilterData.flags & FChaosQueryFlag::eSKIPNARROWPHASE)
+#if PHYSICS_INTERFACE_PHYSX
+	bool bSkipNarrowPhase = false; // Flag doesn't exist on PhysX type.
+#else
+	bool bSkipNarrowPhase = QueryFilterData.flags & FPhysicsQueryFlag::eSKIPNARROWPHASE;
+#endif
+	if (bSkipNarrowPhase)
 	{
 		TBPVisitor<QueryGeomType, TAccelerationStructureHandle<float, 3>, FOverlapHit> OverlapVisitor(GeomPose, HitBuffer, QueryFilterData, QueryCallback, QueryGeom, DebugParams);
 		SpatialAcceleration.Overlap(Bounds, OverlapVisitor);
